Another CIS fix.

diff --git a/Engine/Plugins/Compositing/Composure/Source/Composure/Composure.Build.cs b/Engine/Plugins/Compositing/Composure/Source/Composure/Composure.Build.cs
--- a/Engine/Plugins/Compositing/Composure/Source/Composure/Composure.Build.cs
+++ b/Engine/Plugins/Compositing/Composure/Source/Composure/Composure.Build.cs
@@ -24,7 +24,7 @@ namespace UnrealBuildTool.Rules
 				}
 				);
             
-            if (UEBuildConfiguration.bBuildEditor == true)
+            if (Target.bBuildEditor == true)
             {
                 PrivateDependencyModuleNames.Add("UnrealEd");
             }
