Remove RodElement (and associated files). To be renamed and added in future release.

diff --git a/Engine/Plugins/Compositing/Composure/Config/BaseComposure.ini b/Engine/Plugins/Compositing/Composure/Config/BaseComposure.ini
--- a/Engine/Plugins/Compositing/Composure/Config/BaseComposure.ini
+++ b/Engine/Plugins/Compositing/Composure/Config/BaseComposure.ini
@@ -5,7 +5,7 @@
 +FeaturedElementClasses="/Composure/Blueprints/CompositingElements/BP_CgCaptureCompElement.BP_CgCaptureCompElement_C"
 +FeaturedElementClasses="/Composure/Blueprints/CompositingElements/BP_CgMatteCompElement.BP_CgMatteCompElement_C"
 
-+DefaultElementNames=(("BP_MediaPlateCompElement_C", "media_plate"), ("BP_CgCaptureCompElement_C", "cg_element"), ("BP_CgMatteCompElement_C", "matte_element"), ("BP_DepthCompElement_C", "depth_element"), ("BP_RadialCompElement_C", "radial_element"), ("BP_RampCompElement_C", "ramp_element"), ("BP_RodCompElement_C", "rod_element"), ("BP_RotoPaintCompElement_C", "roto_element"), ("BP_TextureCompElement_C", "texture_element"))
++DefaultElementNames=(("BP_MediaPlateCompElement_C", "media_plate"), ("BP_CgCaptureCompElement_C", "cg_element"), ("BP_CgMatteCompElement_C", "matte_element"), ("BP_DepthCompElement_C", "depth_element"), ("BP_RadialCompElement_C", "radial_element"), ("BP_RampCompElement_C", "ramp_element"), ("BP_RotoPaintCompElement_C", "roto_element"), ("BP_TextureCompElement_C", "texture_element"))
 
 [CoreRedirects]
-+ClassRedirects=(OldName="ComposureCompShotElement",NewName="/Script/Composure.CompositingElement")
\ No newline at end of file
++ClassRedirects=(OldName="ComposureCompShotElement",NewName="/Script/Composure.CompositingElement")
