#UE4 Fix for compile error due to range-based for change

#jria nojira

diff --git a/Engine/Plugins/Compositing/Composure/Source/Composure/Public/CompositingElements/CompositingTextureLookupTable.h b/Engine/Plugins/Compositing/Composure/Source/Composure/Public/CompositingElements/CompositingTextureLookupTable.h
--- a/Engine/Plugins/Compositing/Composure/Source/Composure/Public/CompositingElements/CompositingTextureLookupTable.h
+++ b/Engine/Plugins/Compositing/Composure/Source/Composure/Public/CompositingElements/CompositingTextureLookupTable.h
@@ -43,12 +43,14 @@ private:
 	
 	TMap<FName, ICompositingTextureLookupTable*> LinkedSearchTables;
 
+public:
+
 	/**
 	 * DO NOT USE DIRECTLY
 	 * STL-like iterators to enable range-based for loop support.
 	 */
-	FORCEINLINE friend FInternalLookupTable::TRangedForIterator      begin(      FCompositingTextureLookupTable& InLookupTable) { return begin(InLookupTable.LookupTable); }
-	FORCEINLINE friend FInternalLookupTable::TRangedForConstIterator begin(const FCompositingTextureLookupTable& InLookupTable) { return begin(InLookupTable.LookupTable); }
-	FORCEINLINE friend FInternalLookupTable::TRangedForIterator      end  (      FCompositingTextureLookupTable& InLookupTable) { return end(InLookupTable.LookupTable); }
-	FORCEINLINE friend FInternalLookupTable::TRangedForConstIterator end  (const FCompositingTextureLookupTable& InLookupTable) { return end(InLookupTable.LookupTable); }
+	FORCEINLINE FInternalLookupTable::TRangedForIterator      begin()       { return LookupTable.begin(); }
+	FORCEINLINE FInternalLookupTable::TRangedForConstIterator begin() const { return LookupTable.begin(); }
+	FORCEINLINE FInternalLookupTable::TRangedForIterator      end  ()       { return LookupTable.end(); }
+	FORCEINLINE FInternalLookupTable::TRangedForConstIterator end  () const { return LookupTable.end(); }
 };
\ No newline at end of file
