#UE4 Another CIS fix for after range-based for change

diff --git a/Engine/Plugins/Compositing/Composure/Source/ComposureLayersEditor/Private/EditorCompElementContainer.h b/Engine/Plugins/Compositing/Composure/Source/ComposureLayersEditor/Private/EditorCompElementContainer.h
--- a/Engine/Plugins/Compositing/Composure/Source/ComposureLayersEditor/Private/EditorCompElementContainer.h
+++ b/Engine/Plugins/Compositing/Composure/Source/ComposureLayersEditor/Private/EditorCompElementContainer.h
@@ -48,12 +48,13 @@ private:
 	UPROPERTY()
 	TArray< TWeakObjectPtr<ACompositingElement> > CompElements;
 
+public:
 	/**
 	 * DO NOT USE DIRECTLY
 	 * STL-like iterators to enable range-based for loop support.
 	 */
-	FORCEINLINE friend FCompElementList::RangedForIteratorType      begin(      UEditorCompElementContainer& ElementContainer) { return begin(ElementContainer.CompElements); }
-	FORCEINLINE friend FCompElementList::RangedForConstIteratorType begin(const UEditorCompElementContainer& ElementContainer) { return begin(ElementContainer.CompElements); }
-	FORCEINLINE friend FCompElementList::RangedForIteratorType      end  (      UEditorCompElementContainer& ElementContainer) { return end(ElementContainer.CompElements); }
-	FORCEINLINE friend FCompElementList::RangedForConstIteratorType end  (const UEditorCompElementContainer& ElementContainer) { return end(ElementContainer.CompElements); }
+	FORCEINLINE FCompElementList::RangedForIteratorType      begin()       { return CompElements.begin(); }
+	FORCEINLINE FCompElementList::RangedForConstIteratorType begin() const { return CompElements.begin(); }
+	FORCEINLINE FCompElementList::RangedForIteratorType      end  ()       { return CompElements.end();   }
+	FORCEINLINE FCompElementList::RangedForConstIteratorType end  () const { return CompElements.end();   }
 };
