Fix warnings in SimpleComposure plugin on Mac.
 * Removed unused dependancy on a windows only plugin

diff --git a/Engine/Plugins/Compositing/Composure/Composure.uplugin b/Engine/Plugins/Compositing/Composure/Composure.uplugin
--- a/Engine/Plugins/Compositing/Composure/Composure.uplugin
+++ b/Engine/Plugins/Compositing/Composure/Composure.uplugin
@@ -42,7 +42,7 @@
 			"Enabled": true
 		},
 		{
-			"Name": "MediaFrameworkUtilities",
+			"Name": "MediaIOFramework",
 			"Enabled": true
 		},
 		{
