GitHub #5595 : Replacing uses of FindObject<UEnum> with StaticEnum<>()

diff --git a/Engine/Plugins/Compositing/Composure/Source/ComposureLayersEditor/Private/ComposureDetailCustomizations.cpp b/Engine/Plugins/Compositing/Composure/Source/ComposureLayersEditor/Private/ComposureDetailCustomizations.cpp
--- a/Engine/Plugins/Compositing/Composure/Source/ComposureLayersEditor/Private/ComposureDetailCustomizations.cpp
+++ b/Engine/Plugins/Compositing/Composure/Source/ComposureLayersEditor/Private/ComposureDetailCustomizations.cpp
@@ -409,7 +409,7 @@ void FCompElementDetailsCustomization::ForceRefreshLayout()
 
 void FCompElementDetailsCustomization::GetInstanceCameraSourceComboStrings(TArray<TSharedPtr<FString>>& OutComboBoxStrings, TArray<TSharedPtr<SToolTip>>& OutToolTips, TArray<bool>& OutRestrictedItems)
 {
-	const UEnum* CamSourceEnum = FindObject<UEnum>(ANY_PACKAGE, TEXT("ESceneCameraLinkType"));
+	const UEnum* CamSourceEnum = StaticEnum<ESceneCameraLinkType>();
 	if (ensure(CamSourceEnum))
 	{
 		for (int32 EnumIndex = 0; EnumIndex < CamSourceEnum->NumEnums()-1; ++EnumIndex)
@@ -445,7 +445,7 @@ FString FCompElementDetailsCustomization::GetInstanceCameraSourceValueStr(TShare
 		}
 		else
 		{
-			const UEnum* CamSourceEnum = FindObject<UEnum>(ANY_PACKAGE, TEXT("ESceneCameraLinkType"));
+			const UEnum* CamSourceEnum = StaticEnum<ESceneCameraLinkType>();
 			if (ensure(CamSourceEnum))
 			{
 				DisplayStr = CamSourceEnum->GetDisplayNameTextByValue(CurrentValue).ToString();
@@ -458,7 +458,7 @@ FString FCompElementDetailsCustomization::GetInstanceCameraSourceValueStr(TShare
 
 void FCompElementDetailsCustomization::OnCameraSourceSelected(const FString& Selection, TSharedPtr<IPropertyHandle> PropertyHandle)
 {
-	const UEnum* CamSourceEnum = FindObject<UEnum>(ANY_PACKAGE, TEXT("ESceneCameraLinkType"));
+	const UEnum* CamSourceEnum = StaticEnum<ESceneCameraLinkType>();
 	if (ensure(CamSourceEnum) && PropertyHandle.IsValid())
 	{
 		const int64 FoundValue = CamSourceEnum->GetValueByNameString(Selection);
diff --git a/Engine/Plugins/FX/Niagara/Source/Niagara/Private/NiagaraDataInterfaceCollisionQuery.cpp b/Engine/Plugins/FX/Niagara/Source/Niagara/Private/NiagaraDataInterfaceCollisionQuery.cpp
--- a/Engine/Plugins/FX/Niagara/Source/Niagara/Private/NiagaraDataInterfaceCollisionQuery.cpp
+++ b/Engine/Plugins/FX/Niagara/Source/Niagara/Private/NiagaraDataInterfaceCollisionQuery.cpp
@@ -16,7 +16,7 @@ FCriticalSection UNiagaraDataInterfaceCollisionQuery::CriticalSection;
 UNiagaraDataInterfaceCollisionQuery::UNiagaraDataInterfaceCollisionQuery(FObjectInitializer const& ObjectInitializer)
 	: Super(ObjectInitializer)
 {
-	TraceChannelEnum = FindObject<UEnum>(ANY_PACKAGE, TEXT("ECollisionChannel"), true);
+	TraceChannelEnum = StaticEnum<ECollisionChannel>();
 }
 
 bool UNiagaraDataInterfaceCollisionQuery::InitPerInstanceData(void* PerInstanceData, FNiagaraSystemInstance* InSystemInstance)
diff --git a/Engine/Plugins/FX/Niagara/Source/NiagaraEditor/Private/NiagaraNodeUsageSelector.cpp b/Engine/Plugins/FX/Niagara/Source/NiagaraEditor/Private/NiagaraNodeUsageSelector.cpp
--- a/Engine/Plugins/FX/Niagara/Source/NiagaraEditor/Private/NiagaraNodeUsageSelector.cpp
+++ b/Engine/Plugins/FX/Niagara/Source/NiagaraEditor/Private/NiagaraNodeUsageSelector.cpp
@@ -35,7 +35,7 @@ bool UNiagaraNodeUsageSelector::AllowNiagaraTypeForAddPin(const FNiagaraTypeDefi
 void UNiagaraNodeUsageSelector::InsertInputPinsFor(const FNiagaraVariable& Var)
 {
 	const UEdGraphSchema_Niagara* Schema = GetDefault<UEdGraphSchema_Niagara>();
-	UEnum* ENiagaraScriptGroupEnum = FindObject<UEnum>(ANY_PACKAGE, TEXT("ENiagaraScriptGroup"), true);
+	UEnum* ENiagaraScriptGroupEnum = StaticEnum<ENiagaraScriptGroup>();
 	int64 GroupCount = (int64)ENiagaraScriptGroup::Max;
 
 	TArray<UEdGraphPin*> OldPins(Pins);
diff --git a/Engine/Plugins/Runtime/AR/Apple/AppleARKitFaceSupport/Source/AppleARKitFaceSupport/Private/AppleARKitLiveLinkSource.cpp b/Engine/Plugins/Runtime/AR/Apple/AppleARKitFaceSupport/Source/AppleARKitFaceSupport/Private/AppleARKitLiveLinkSource.cpp
--- a/Engine/Plugins/Runtime/AR/Apple/AppleARKitFaceSupport/Source/AppleARKitFaceSupport/Private/AppleARKitLiveLinkSource.cpp
+++ b/Engine/Plugins/Runtime/AR/Apple/AppleARKitFaceSupport/Source/AppleARKitFaceSupport/Private/AppleARKitLiveLinkSource.cpp
@@ -512,7 +512,7 @@ FAppleARKitLiveLinkFileWriterCsv::FAppleARKitLiveLinkFileWriterCsv()
 	check(IsInGameThread());
 
 	CsvFrameHeader = TEXT("Timecode, FrameRate");
-	const UEnum *EnumPtr = FindObject<UEnum>(ANY_PACKAGE, TEXT("EARFaceBlendShape"), true);
+	const UEnum *EnumPtr = StaticEnum<EARFaceBlendShape>();
 	if (EnumPtr != nullptr)
 	{
 		// Iterate through all of the enum values generating strings for them for CSV/JSON generation
@@ -560,7 +560,7 @@ FAppleARKitLiveLinkFileWriterJson::FAppleARKitLiveLinkFileWriterJson()
 	// Touching UObjects, so needs to be game thread
 	check(IsInGameThread());
 
-	const UEnum *EnumPtr = FindObject<UEnum>(ANY_PACKAGE, TEXT("EARFaceBlendShape"), true);
+	const UEnum *EnumPtr = StaticEnum<EARFaceBlendShape>();
 	if (EnumPtr != nullptr)
 	{
 		// Iterate through all of the enum values generating strings for them for CSV/JSON generation
diff --git a/Engine/Plugins/Runtime/PacketHandlers/CompressionComponents/Oodle/Source/OodleHandlerComponent/Private/OodleHandlerComponent.cpp b/Engine/Plugins/Runtime/PacketHandlers/CompressionComponents/Oodle/Source/OodleHandlerComponent/Private/OodleHandlerComponent.cpp
--- a/Engine/Plugins/Runtime/PacketHandlers/CompressionComponents/Oodle/Source/OodleHandlerComponent/Private/OodleHandlerComponent.cpp
+++ b/Engine/Plugins/Runtime/PacketHandlers/CompressionComponents/Oodle/Source/OodleHandlerComponent/Private/OodleHandlerComponent.cpp
@@ -353,7 +353,7 @@ void OodleHandlerComponent::Initialize()
 	}
 
 	FString CurEnableModeStr;
-	UEnum* EnableModeEnum = FindObject<UEnum>(ANY_PACKAGE, TEXT("EOodleEnableMode"), true);
+	UEnum* EnableModeEnum = StaticEnum<EOodleEnableMode>();
 
 	if (GConfig->GetString(OODLE_INI_SECTION, TEXT("ServerEnableMode"), CurEnableModeStr, GEngineIni))
 	{
diff --git a/Engine/Plugins/Runtime/nDisplay/Source/DisplayClusterInput/Private/Blueprints/DisplayClusterInputBlueprintAPIImpl.cpp b/Engine/Plugins/Runtime/nDisplay/Source/DisplayClusterInput/Private/Blueprints/DisplayClusterInputBlueprintAPIImpl.cpp
--- a/Engine/Plugins/Runtime/nDisplay/Source/DisplayClusterInput/Private/Blueprints/DisplayClusterInputBlueprintAPIImpl.cpp
+++ b/Engine/Plugins/Runtime/nDisplay/Source/DisplayClusterInput/Private/Blueprints/DisplayClusterInputBlueprintAPIImpl.cpp
@@ -64,14 +64,6 @@ bool UDisplayClusterInputBlueprintAPIImpl::BindVrpnTracker(const FString& VrpnDe
 {
 	DISPLAY_CLUSTER_FUNC_TRACE(LogDisplayClusterInputBP);
 
-	const FString ObjName(TEXT("EControllerHand"));
-	const UEnum* const EnumPtr = FindObject<UEnum>(ANY_PACKAGE, *ObjName, true);
-	if (!EnumPtr)
-	{
-		UE_LOG(LogDisplayClusterInputBP, Error, TEXT("Couldn't find %s object"), *ObjName);
-		return false;
-	}
-
-	FName TargetName = EnumPtr->GetNameByValue((int64)Target);
-	return IDisplayClusterInputModule::Get().BindVrpnChannel(VrpnDeviceId, VrpnChannel, TargetName.ToString());
+	const FString TargetString = UEnum::GetValueAsString<EControllerHand>(Target);
+	return IDisplayClusterInputModule::Get().BindVrpnChannel(VrpnDeviceId, VrpnChannel, TargetString);
 }
diff --git a/Engine/Source/Editor/UnrealEd/Private/Fbx/FbxMainImport.cpp b/Engine/Source/Editor/UnrealEd/Private/Fbx/FbxMainImport.cpp
--- a/Engine/Source/Editor/UnrealEd/Private/Fbx/FbxMainImport.cpp
+++ b/Engine/Source/Editor/UnrealEd/Private/Fbx/FbxMainImport.cpp
@@ -1475,11 +1475,11 @@ bool FFbxImporter::ImportFromFile(const FString& Filename, const FString& Type,
 					if( FEngineAnalytics::IsAvailable() )
 					{
 						const static UEnum* FBXImportTypeEnum = StaticEnum<EFBXImportType>();
-						const static UEnum* FBXAnimationLengthImportTypeEnum = FindObject<UEnum>(ANY_PACKAGE, TEXT("EFBXAnimationLengthImportType"));
-						const static UEnum* MaterialSearchLocationEnum = FindObject<UEnum>(ANY_PACKAGE, TEXT("EMaterialSearchLocation"));
-						const static UEnum* FBXNormalGenerationMethodEnum = FindObject<UEnum>(ANY_PACKAGE, TEXT("EFBXNormalGenerationMethod"));
-						const static UEnum* FBXNormalImportMethodEnum = FindObject<UEnum>(ANY_PACKAGE, TEXT("EFBXNormalImportMethod"));
-						const static UEnum* VertexColorImportOptionEnum = FindObject<UEnum>(ANY_PACKAGE, TEXT("EVertexColorImportOption"));
+						const static UEnum* FBXAnimationLengthImportTypeEnum = StaticEnum<EFBXAnimationLengthImportType>();
+						const static UEnum* MaterialSearchLocationEnum = StaticEnum<EMaterialSearchLocation>();
+						const static UEnum* FBXNormalGenerationMethodEnum = StaticEnum<EFBXNormalGenerationMethod::Type>();
+						const static UEnum* FBXNormalImportMethodEnum = StaticEnum<EFBXNormalImportMethod>();
+						const static UEnum* VertexColorImportOptionEnum = StaticEnum<EVertexColorImportOption::Type>();
 						
 						TArray<FAnalyticsEventAttribute> Attribs;
 
