Sequencer: Remove null sections on PostLoad. Added RemoveSectionAt given a section index.

diff --git a/Engine/Plugins/Compositing/Composure/Source/Composure/Public/MovieScene/MovieSceneComposureExportTrack.h b/Engine/Plugins/Compositing/Composure/Source/Composure/Public/MovieScene/MovieSceneComposureExportTrack.h
--- a/Engine/Plugins/Compositing/Composure/Source/Composure/Public/MovieScene/MovieSceneComposureExportTrack.h
+++ b/Engine/Plugins/Compositing/Composure/Source/Composure/Public/MovieScene/MovieSceneComposureExportTrack.h
@@ -53,7 +53,8 @@ public:
 	virtual void AddSection(UMovieSceneSection& Section) override              { Sections.AddUnique(&Section); }
 	virtual const TArray<UMovieSceneSection*>& GetAllSections() const override { return Sections; }
 	virtual bool HasSection(const UMovieSceneSection& Section) const override  { return Sections.Contains(&Section); }
-	virtual void RemoveSection(UMovieSceneSection& Section) override           { Sections.Remove(&Section); }
+	virtual void RemoveSection(UMovieSceneSection& Section) override { Sections.Remove(&Section); }
+	virtual void RemoveSectionAt(int32 SectionIndex) override { Sections.RemoveAt(SectionIndex); }
 	virtual bool SupportsMultipleRows() const override                         { return true; }
 	virtual FMovieSceneEvalTemplatePtr CreateTemplateForSection(const UMovieSceneSection& InSection) const override;
 
diff --git a/Engine/Plugins/Experimental/GeometryCache/Source/GeometryCacheTracks/Classes/MovieSceneGeometryCacheTrack.h b/Engine/Plugins/Experimental/GeometryCache/Source/GeometryCacheTracks/Classes/MovieSceneGeometryCacheTrack.h
--- a/Engine/Plugins/Experimental/GeometryCache/Source/GeometryCacheTracks/Classes/MovieSceneGeometryCacheTrack.h
+++ b/Engine/Plugins/Experimental/GeometryCache/Source/GeometryCacheTracks/Classes/MovieSceneGeometryCacheTrack.h
@@ -33,6 +33,7 @@ public:
 	virtual bool HasSection(const UMovieSceneSection& Section) const override;
 	virtual void AddSection(UMovieSceneSection& Section) override;
 	virtual void RemoveSection(UMovieSceneSection& Section) override;
+	virtual void RemoveSectionAt(int32 SectionIndex) override;
 	virtual bool IsEmpty() const override;
 	virtual const TArray<UMovieSceneSection*>& GetAllSections() const override;
 	virtual bool SupportsType(TSubclassOf<UMovieSceneSection> SectionClass) const override;
diff --git a/Engine/Plugins/Experimental/GeometryCache/Source/GeometryCacheTracks/Private/MovieSceneGeometryCacheTrack.cpp b/Engine/Plugins/Experimental/GeometryCache/Source/GeometryCacheTracks/Private/MovieSceneGeometryCacheTrack.cpp
--- a/Engine/Plugins/Experimental/GeometryCache/Source/GeometryCacheTracks/Private/MovieSceneGeometryCacheTrack.cpp
+++ b/Engine/Plugins/Experimental/GeometryCache/Source/GeometryCacheTracks/Private/MovieSceneGeometryCacheTrack.cpp
@@ -106,6 +106,11 @@ void UMovieSceneGeometryCacheTrack::RemoveSection(UMovieSceneSection& Section)
 	AnimationSections.Remove(&Section);
 }
 
+void UMovieSceneGeometryCacheTrack::RemoveSectionAt(int32 SectionIndex)
+{
+	AnimationSections.RemoveAt(SectionIndex);
+}
+
 
 bool UMovieSceneGeometryCacheTrack::IsEmpty() const
 {
diff --git a/Engine/Plugins/Experimental/ImagePlate/Source/ImagePlate/Private/MovieSceneImagePlateTrack.cpp b/Engine/Plugins/Experimental/ImagePlate/Source/ImagePlate/Private/MovieSceneImagePlateTrack.cpp
--- a/Engine/Plugins/Experimental/ImagePlate/Source/ImagePlate/Private/MovieSceneImagePlateTrack.cpp
+++ b/Engine/Plugins/Experimental/ImagePlate/Source/ImagePlate/Private/MovieSceneImagePlateTrack.cpp
@@ -42,6 +42,11 @@ void UMovieSceneImagePlateTrack::RemoveSection(UMovieSceneSection& Section)
 	Sections.Remove(&Section);
 }
 
+void UMovieSceneImagePlateTrack::RemoveSectionAt(int32 SectionIndex)
+{
+	Sections.RemoveAt(SectionIndex);
+}
+
 FMovieSceneEvalTemplatePtr UMovieSceneImagePlateTrack::CreateTemplateForSection(const UMovieSceneSection& InSection) const
 {
 	return FMovieSceneImagePlateSectionTemplate(*CastChecked<const UMovieSceneImagePlateSection>(&InSection), *this);
diff --git a/Engine/Plugins/Experimental/ImagePlate/Source/ImagePlate/Public/MovieSceneImagePlateTrack.h b/Engine/Plugins/Experimental/ImagePlate/Source/ImagePlate/Public/MovieSceneImagePlateTrack.h
--- a/Engine/Plugins/Experimental/ImagePlate/Source/ImagePlate/Public/MovieSceneImagePlateTrack.h
+++ b/Engine/Plugins/Experimental/ImagePlate/Source/ImagePlate/Public/MovieSceneImagePlateTrack.h
@@ -27,6 +27,7 @@ public:
 	virtual UMovieSceneSection* CreateNewSection() override;
 	virtual const TArray<UMovieSceneSection*>& GetAllSections() const override;
 	virtual void RemoveSection(UMovieSceneSection& Section) override;
+	virtual void RemoveSectionAt(int32 SectionIndex) override;
 	virtual FMovieSceneEvalTemplatePtr CreateTemplateForSection(const UMovieSceneSection& InSection) const override;
 #if WITH_EDITORONLY_DATA
 	virtual FText GetDefaultDisplayName() const override;
diff --git a/Engine/Plugins/FX/Niagara/Source/Niagara/Private/MovieScene/MovieSceneNiagaraTrack.cpp b/Engine/Plugins/FX/Niagara/Source/Niagara/Private/MovieScene/MovieSceneNiagaraTrack.cpp
--- a/Engine/Plugins/FX/Niagara/Source/Niagara/Private/MovieScene/MovieSceneNiagaraTrack.cpp
+++ b/Engine/Plugins/FX/Niagara/Source/Niagara/Private/MovieScene/MovieSceneNiagaraTrack.cpp
@@ -31,3 +31,8 @@ void UMovieSceneNiagaraTrack::RemoveSection(UMovieSceneSection& Section)
 {
 	Sections.Remove(&Section);
 }
+
+void UMovieSceneNiagaraTrack::RemoveSectionAt(int32 SectionIndex)
+{
+	Sections.RemoveAt(SectionIndex);
+}
diff --git a/Engine/Plugins/FX/Niagara/Source/Niagara/Public/MovieScene/MovieSceneNiagaraTrack.h b/Engine/Plugins/FX/Niagara/Source/Niagara/Public/MovieScene/MovieSceneNiagaraTrack.h
--- a/Engine/Plugins/FX/Niagara/Source/Niagara/Public/MovieScene/MovieSceneNiagaraTrack.h
+++ b/Engine/Plugins/FX/Niagara/Source/Niagara/Public/MovieScene/MovieSceneNiagaraTrack.h
@@ -17,6 +17,7 @@ public:
 	virtual bool IsEmpty() const override;
 	virtual void RemoveAllAnimationData() override;
 	virtual void RemoveSection(UMovieSceneSection& Section) override;
+	virtual void RemoveSectionAt(int32 SectionIndex) override;
 
 protected:
 	UPROPERTY()
diff --git a/Engine/Plugins/FX/Niagara/Source/NiagaraEditor/Private/Sequencer/NiagaraSequence/MovieSceneNiagaraEmitterTrack.cpp b/Engine/Plugins/FX/Niagara/Source/NiagaraEditor/Private/Sequencer/NiagaraSequence/MovieSceneNiagaraEmitterTrack.cpp
--- a/Engine/Plugins/FX/Niagara/Source/NiagaraEditor/Private/Sequencer/NiagaraSequence/MovieSceneNiagaraEmitterTrack.cpp
+++ b/Engine/Plugins/FX/Niagara/Source/NiagaraEditor/Private/Sequencer/NiagaraSequence/MovieSceneNiagaraEmitterTrack.cpp
@@ -164,6 +164,11 @@ void UMovieSceneNiagaraEmitterTrack::RemoveSection(UMovieSceneSection& Section)
 	Sections.Remove(&Section);
 }
 
+void UMovieSceneNiagaraEmitterTrack::RemoveSectionAt(int32 SectionIndex)
+{
+	Sections.RemoveAt(SectionIndex);
+}
+
 bool UMovieSceneNiagaraEmitterTrack::IsEmpty() const
 {
 	return Sections.Num() == 0;
diff --git a/Engine/Plugins/FX/Niagara/Source/NiagaraEditor/Private/Sequencer/NiagaraSequence/MovieSceneNiagaraEmitterTrack.h b/Engine/Plugins/FX/Niagara/Source/NiagaraEditor/Private/Sequencer/NiagaraSequence/MovieSceneNiagaraEmitterTrack.h
--- a/Engine/Plugins/FX/Niagara/Source/NiagaraEditor/Private/Sequencer/NiagaraSequence/MovieSceneNiagaraEmitterTrack.h
+++ b/Engine/Plugins/FX/Niagara/Source/NiagaraEditor/Private/Sequencer/NiagaraSequence/MovieSceneNiagaraEmitterTrack.h
@@ -74,6 +74,7 @@ public:
 	virtual void RemoveAllAnimationData() override { }
 	virtual bool HasSection(const UMovieSceneSection& Section) const override;
 	virtual void RemoveSection(UMovieSceneSection& Section) override;
+	virtual void RemoveSectionAt(int32 SectionIndex) override;
 	virtual bool IsEmpty() const override;
 	virtual const TArray<UMovieSceneSection*>& GetAllSections() const override;
 	virtual bool SupportsMultipleRows() const override;
diff --git a/Engine/Plugins/Media/MediaCompositing/Source/MediaCompositing/Private/MovieScene/MovieSceneMediaTrack.cpp b/Engine/Plugins/Media/MediaCompositing/Source/MediaCompositing/Private/MovieScene/MovieSceneMediaTrack.cpp
--- a/Engine/Plugins/Media/MediaCompositing/Source/MediaCompositing/Private/MovieScene/MovieSceneMediaTrack.cpp
+++ b/Engine/Plugins/Media/MediaCompositing/Source/MediaCompositing/Private/MovieScene/MovieSceneMediaTrack.cpp
@@ -96,6 +96,11 @@ void UMovieSceneMediaTrack::RemoveSection(UMovieSceneSection& Section)
 	MediaSections.Remove(&Section);
 }
 
+void UMovieSceneMediaTrack::RemoveSectionAt(int32 SectionIndex)
+{
+	MediaSections.RemoveAt(SectionIndex);
+}
+
 
 FMovieSceneEvalTemplatePtr UMovieSceneMediaTrack::CreateTemplateForSection(const UMovieSceneSection& InSection) const
 {
diff --git a/Engine/Plugins/Media/MediaCompositing/Source/MediaCompositing/Public/MovieSceneMediaTrack.h b/Engine/Plugins/Media/MediaCompositing/Source/MediaCompositing/Public/MovieSceneMediaTrack.h
--- a/Engine/Plugins/Media/MediaCompositing/Source/MediaCompositing/Public/MovieSceneMediaTrack.h
+++ b/Engine/Plugins/Media/MediaCompositing/Source/MediaCompositing/Public/MovieSceneMediaTrack.h
@@ -49,6 +49,7 @@ public:
 	virtual bool HasSection(const UMovieSceneSection& Section) const override;
 	virtual bool IsEmpty() const override;
 	virtual void RemoveSection(UMovieSceneSection& Section) override;
+	virtual void RemoveSectionAt(int32 SectionIndex) override;
 	virtual bool SupportsMultipleRows() const override { return true; }
 
 private:
diff --git a/Engine/Source/Runtime/MovieScene/Private/Tracks/MovieSceneSubTrack.cpp b/Engine/Source/Runtime/MovieScene/Private/Tracks/MovieSceneSubTrack.cpp
--- a/Engine/Source/Runtime/MovieScene/Private/Tracks/MovieSceneSubTrack.cpp
+++ b/Engine/Source/Runtime/MovieScene/Private/Tracks/MovieSceneSubTrack.cpp
@@ -192,6 +192,11 @@ void UMovieSceneSubTrack::RemoveSection(UMovieSceneSection& Section)
 	Sections.Remove(&Section);
 }
 
+void UMovieSceneSubTrack::RemoveSectionAt(int32 SectionIndex)
+{
+	Sections.RemoveAt(SectionIndex);
+}
+
 
 bool UMovieSceneSubTrack::SupportsMultipleRows() const
 {
diff --git a/Engine/Source/Runtime/MovieScene/Public/MovieSceneTrack.h b/Engine/Source/Runtime/MovieScene/Public/MovieSceneTrack.h
--- a/Engine/Source/Runtime/MovieScene/Public/MovieSceneTrack.h
+++ b/Engine/Source/Runtime/MovieScene/Public/MovieSceneTrack.h
@@ -299,7 +299,14 @@ public:
 	 *
 	 * @param Section The section to remove.
 	 */
-	virtual void RemoveSection(UMovieSceneSection& Section) PURE_VIRTUAL(UMovieSceneSection::RemoveSection,);
+	virtual void RemoveSection(UMovieSceneSection& Section) PURE_VIRTUAL(UMovieSceneSection::RemoveSection, );
+
+	/**
+	 * Removes a section from this track at a particular index
+	 *
+	 * @param SectionIndex The section index to remove.
+	 */
+	virtual void RemoveSectionAt(int32 SectionIndex) PURE_VIRTUAL(UMovieSceneSection::RemoveSectionAt, );
 
 #if WITH_EDITOR
 
diff --git a/Engine/Source/Runtime/MovieScene/Public/Tracks/MovieSceneSubTrack.h b/Engine/Source/Runtime/MovieScene/Public/Tracks/MovieSceneSubTrack.h
--- a/Engine/Source/Runtime/MovieScene/Public/Tracks/MovieSceneSubTrack.h
+++ b/Engine/Source/Runtime/MovieScene/Public/Tracks/MovieSceneSubTrack.h
@@ -72,6 +72,7 @@ public:
 	virtual bool IsEmpty() const override;
 	virtual void RemoveAllAnimationData() override;
 	virtual void RemoveSection(UMovieSceneSection& Section) override;
+	virtual void RemoveSectionAt(int32 SectionIndex) override;
 	virtual bool SupportsMultipleRows() const override;
 
 #if WITH_EDITORONLY_DATA
diff --git a/Engine/Source/Runtime/MovieSceneTracks/Private/Tracks/MovieScene3DConstraintTrack.cpp b/Engine/Source/Runtime/MovieSceneTracks/Private/Tracks/MovieScene3DConstraintTrack.cpp
--- a/Engine/Source/Runtime/MovieSceneTracks/Private/Tracks/MovieScene3DConstraintTrack.cpp
+++ b/Engine/Source/Runtime/MovieSceneTracks/Private/Tracks/MovieScene3DConstraintTrack.cpp
@@ -47,6 +47,12 @@ void UMovieScene3DConstraintTrack::RemoveSection(UMovieSceneSection& Section)
 }
 
 
+void UMovieScene3DConstraintTrack::RemoveSectionAt(int32 SectionIndex)
+{
+	ConstraintSections.RemoveAt(SectionIndex);
+}
+
+
 bool UMovieScene3DConstraintTrack::IsEmpty() const
 {
 	return ConstraintSections.Num() == 0;
diff --git a/Engine/Source/Runtime/MovieSceneTracks/Private/Tracks/MovieSceneAudioTrack.cpp b/Engine/Source/Runtime/MovieSceneTracks/Private/Tracks/MovieSceneAudioTrack.cpp
--- a/Engine/Source/Runtime/MovieSceneTracks/Private/Tracks/MovieSceneAudioTrack.cpp
+++ b/Engine/Source/Runtime/MovieSceneTracks/Private/Tracks/MovieSceneAudioTrack.cpp
@@ -65,6 +65,12 @@ void UMovieSceneAudioTrack::RemoveSection(UMovieSceneSection& Section)
 }
 
 
+void UMovieSceneAudioTrack::RemoveSectionAt(int32 SectionIndex)
+{
+	AudioSections.RemoveAt(SectionIndex);
+}
+
+
 bool UMovieSceneAudioTrack::IsEmpty() const
 {
 	return AudioSections.Num() == 0;
diff --git a/Engine/Source/Runtime/MovieSceneTracks/Private/Tracks/MovieSceneCameraAnimTrack.cpp b/Engine/Source/Runtime/MovieSceneTracks/Private/Tracks/MovieSceneCameraAnimTrack.cpp
--- a/Engine/Source/Runtime/MovieSceneTracks/Private/Tracks/MovieSceneCameraAnimTrack.cpp
+++ b/Engine/Source/Runtime/MovieSceneTracks/Private/Tracks/MovieSceneCameraAnimTrack.cpp
@@ -78,6 +78,11 @@ void UMovieSceneCameraAnimTrack::RemoveSection(UMovieSceneSection& Section)
 	CameraAnimSections.Remove(&Section);
 }
 
+void UMovieSceneCameraAnimTrack::RemoveSectionAt(int32 SectionIndex)
+{
+	CameraAnimSections.RemoveAt(SectionIndex);
+}
+
 
 bool UMovieSceneCameraAnimTrack::IsEmpty() const
 {
diff --git a/Engine/Source/Runtime/MovieSceneTracks/Private/Tracks/MovieSceneCameraCutTrack.cpp b/Engine/Source/Runtime/MovieSceneTracks/Private/Tracks/MovieSceneCameraCutTrack.cpp
--- a/Engine/Source/Runtime/MovieSceneTracks/Private/Tracks/MovieSceneCameraCutTrack.cpp
+++ b/Engine/Source/Runtime/MovieSceneTracks/Private/Tracks/MovieSceneCameraCutTrack.cpp
@@ -102,6 +102,15 @@ void UMovieSceneCameraCutTrack::RemoveSection(UMovieSceneSection& Section)
 	// @todo Sequencer: The movie scene owned by the section is now abandoned.  Should we offer to delete it?  
 }
 
+void UMovieSceneCameraCutTrack::RemoveSectionAt(int32 SectionIndex)
+{
+	UMovieSceneSection* SectionToDelete = Sections[SectionIndex];
+	MovieSceneHelpers::FixupConsecutiveSections(Sections, *SectionToDelete, true);
+
+	Sections.RemoveAt(SectionIndex);
+	MovieSceneHelpers::SortConsecutiveSections(Sections);
+}
+
 
 void UMovieSceneCameraCutTrack::RemoveAllAnimationData()
 {
diff --git a/Engine/Source/Runtime/MovieSceneTracks/Private/Tracks/MovieSceneCameraShakeTrack.cpp b/Engine/Source/Runtime/MovieSceneTracks/Private/Tracks/MovieSceneCameraShakeTrack.cpp
--- a/Engine/Source/Runtime/MovieSceneTracks/Private/Tracks/MovieSceneCameraShakeTrack.cpp
+++ b/Engine/Source/Runtime/MovieSceneTracks/Private/Tracks/MovieSceneCameraShakeTrack.cpp
@@ -88,6 +88,12 @@ void UMovieSceneCameraShakeTrack::RemoveSection(UMovieSceneSection& Section)
 }
 
 
+void UMovieSceneCameraShakeTrack::RemoveSectionAt(int32 SectionIndex)
+{
+	CameraShakeSections.RemoveAt(SectionIndex);
+}
+
+
 bool UMovieSceneCameraShakeTrack::IsEmpty() const
 {
 	return CameraShakeSections.Num() == 0;
diff --git a/Engine/Source/Runtime/MovieSceneTracks/Private/Tracks/MovieSceneCinematicShotTrack.cpp b/Engine/Source/Runtime/MovieSceneTracks/Private/Tracks/MovieSceneCinematicShotTrack.cpp
--- a/Engine/Source/Runtime/MovieSceneTracks/Private/Tracks/MovieSceneCinematicShotTrack.cpp
+++ b/Engine/Source/Runtime/MovieSceneTracks/Private/Tracks/MovieSceneCinematicShotTrack.cpp
@@ -75,6 +75,12 @@ void UMovieSceneCinematicShotTrack::RemoveSection(UMovieSceneSection& Section)
 	// @todo Sequencer: The movie scene owned by the section is now abandoned.  Should we offer to delete it?  
 }
 
+void UMovieSceneCinematicShotTrack::RemoveSectionAt(int32 SectionIndex)
+{
+	Sections.RemoveAt(SectionIndex);
+	MovieSceneHelpers::SortConsecutiveSections(Sections);
+}
+
 bool UMovieSceneCinematicShotTrack::SupportsMultipleRows() const
 {
 	return true;
diff --git a/Engine/Source/Runtime/MovieSceneTracks/Private/Tracks/MovieSceneEventTrack.cpp b/Engine/Source/Runtime/MovieSceneTracks/Private/Tracks/MovieSceneEventTrack.cpp
--- a/Engine/Source/Runtime/MovieSceneTracks/Private/Tracks/MovieSceneEventTrack.cpp
+++ b/Engine/Source/Runtime/MovieSceneTracks/Private/Tracks/MovieSceneEventTrack.cpp
@@ -62,6 +62,11 @@ void UMovieSceneEventTrack::RemoveSection(UMovieSceneSection& Section)
 	Sections.Remove(&Section);
 }
 
+void UMovieSceneEventTrack::RemoveSectionAt(int32 SectionIndex)
+{
+	Sections.RemoveAt(SectionIndex);
+}
+
 FMovieSceneEvalTemplatePtr UMovieSceneEventTrack::CreateTemplateForSection(const UMovieSceneSection& InSection) const
 {
 	if (const UMovieSceneEventSection* LegacyEventSection = Cast<const UMovieSceneEventSection>(&InSection))
diff --git a/Engine/Source/Runtime/MovieSceneTracks/Private/Tracks/MovieSceneLevelVisibilityTrack.cpp b/Engine/Source/Runtime/MovieSceneTracks/Private/Tracks/MovieSceneLevelVisibilityTrack.cpp
--- a/Engine/Source/Runtime/MovieSceneTracks/Private/Tracks/MovieSceneLevelVisibilityTrack.cpp
+++ b/Engine/Source/Runtime/MovieSceneTracks/Private/Tracks/MovieSceneLevelVisibilityTrack.cpp
@@ -35,11 +35,16 @@ void UMovieSceneLevelVisibilityTrack::AddSection(UMovieSceneSection& Section)
 }
 
 
-void UMovieSceneLevelVisibilityTrack::RemoveSection( UMovieSceneSection& Section )
+void UMovieSceneLevelVisibilityTrack::RemoveSection(UMovieSceneSection& Section)
 {
 	Sections.Remove(&Section);
 }
 
+void UMovieSceneLevelVisibilityTrack::RemoveSectionAt(int32 SectionIndex)
+{
+	Sections.RemoveAt(SectionIndex);
+}
+
 bool UMovieSceneLevelVisibilityTrack::SupportsType(TSubclassOf<UMovieSceneSection> SectionClass) const
 {
 	return SectionClass == UMovieSceneLevelVisibilitySection::StaticClass();
diff --git a/Engine/Source/Runtime/MovieSceneTracks/Private/Tracks/MovieSceneMaterialTrack.cpp b/Engine/Source/Runtime/MovieSceneTracks/Private/Tracks/MovieSceneMaterialTrack.cpp
--- a/Engine/Source/Runtime/MovieSceneTracks/Private/Tracks/MovieSceneMaterialTrack.cpp
+++ b/Engine/Source/Runtime/MovieSceneTracks/Private/Tracks/MovieSceneMaterialTrack.cpp
@@ -48,6 +48,11 @@ void UMovieSceneMaterialTrack::RemoveSection(UMovieSceneSection& Section)
 	Sections.Remove(&Section);
 }
 
+void UMovieSceneMaterialTrack::RemoveSectionAt(int32 SectionIndex)
+{
+	Sections.RemoveAt(SectionIndex);
+}
+
 
 bool UMovieSceneMaterialTrack::IsEmpty() const
 {
diff --git a/Engine/Source/Runtime/MovieSceneTracks/Private/Tracks/MovieSceneParticleParameterTrack.cpp b/Engine/Source/Runtime/MovieSceneTracks/Private/Tracks/MovieSceneParticleParameterTrack.cpp
--- a/Engine/Source/Runtime/MovieSceneTracks/Private/Tracks/MovieSceneParticleParameterTrack.cpp
+++ b/Engine/Source/Runtime/MovieSceneTracks/Private/Tracks/MovieSceneParticleParameterTrack.cpp
@@ -49,6 +49,11 @@ void UMovieSceneParticleParameterTrack::RemoveSection(UMovieSceneSection& Sectio
 	Sections.Remove(&Section);
 }
 
+void UMovieSceneParticleParameterTrack::RemoveSectionAt(int32 SectionIndex)
+{
+	Sections.RemoveAt(SectionIndex);
+}
+
 bool UMovieSceneParticleParameterTrack::IsEmpty() const
 {
 	return Sections.Num() == 0;
diff --git a/Engine/Source/Runtime/MovieSceneTracks/Private/Tracks/MovieSceneParticleTrack.cpp b/Engine/Source/Runtime/MovieSceneTracks/Private/Tracks/MovieSceneParticleTrack.cpp
--- a/Engine/Source/Runtime/MovieSceneTracks/Private/Tracks/MovieSceneParticleTrack.cpp
+++ b/Engine/Source/Runtime/MovieSceneTracks/Private/Tracks/MovieSceneParticleTrack.cpp
@@ -47,6 +47,12 @@ void UMovieSceneParticleTrack::RemoveSection(UMovieSceneSection& Section)
 }
 
 
+void UMovieSceneParticleTrack::RemoveSectionAt(int32 SectionIndex)
+{
+	ParticleSections.RemoveAt(SectionIndex);
+}
+
+
 bool UMovieSceneParticleTrack::IsEmpty() const
 {
 	return ParticleSections.Num() == 0;
diff --git a/Engine/Source/Runtime/MovieSceneTracks/Private/Tracks/MovieScenePropertyTrack.cpp b/Engine/Source/Runtime/MovieSceneTracks/Private/Tracks/MovieScenePropertyTrack.cpp
--- a/Engine/Source/Runtime/MovieSceneTracks/Private/Tracks/MovieScenePropertyTrack.cpp
+++ b/Engine/Source/Runtime/MovieSceneTracks/Private/Tracks/MovieScenePropertyTrack.cpp
@@ -92,6 +92,18 @@ void UMovieScenePropertyTrack::RemoveSection(UMovieSceneSection& Section)
 	}
 }
 
+void UMovieScenePropertyTrack::RemoveSectionAt(int32 SectionIndex)
+{
+	bool bResetSectionToKey = SectionToKey = Sections[SectionIndex];
+
+	Sections.RemoveAt(SectionIndex);
+
+	if (bResetSectionToKey)
+	{
+		SectionToKey = Sections.Num() > 0 ? Sections[0] : nullptr;
+	}
+}
+
 
 bool UMovieScenePropertyTrack::IsEmpty() const
 {
diff --git a/Engine/Source/Runtime/MovieSceneTracks/Private/Tracks/MovieSceneSkeletalAnimationTrack.cpp b/Engine/Source/Runtime/MovieSceneTracks/Private/Tracks/MovieSceneSkeletalAnimationTrack.cpp
--- a/Engine/Source/Runtime/MovieSceneTracks/Private/Tracks/MovieSceneSkeletalAnimationTrack.cpp
+++ b/Engine/Source/Runtime/MovieSceneTracks/Private/Tracks/MovieSceneSkeletalAnimationTrack.cpp
@@ -122,6 +122,12 @@ void UMovieSceneSkeletalAnimationTrack::RemoveSection(UMovieSceneSection& Sectio
 }
 
 
+void UMovieSceneSkeletalAnimationTrack::RemoveSectionAt(int32 SectionIndex)
+{
+	AnimationSections.RemoveAt(SectionIndex);
+}
+
+
 bool UMovieSceneSkeletalAnimationTrack::IsEmpty() const
 {
 	return AnimationSections.Num() == 0;
diff --git a/Engine/Source/Runtime/MovieSceneTracks/Private/Tracks/MovieSceneSpawnTrack.cpp b/Engine/Source/Runtime/MovieSceneTracks/Private/Tracks/MovieSceneSpawnTrack.cpp
--- a/Engine/Source/Runtime/MovieSceneTracks/Private/Tracks/MovieSceneSpawnTrack.cpp
+++ b/Engine/Source/Runtime/MovieSceneTracks/Private/Tracks/MovieSceneSpawnTrack.cpp
@@ -77,7 +77,12 @@ void UMovieSceneSpawnTrack::AddSection(UMovieSceneSection& Section)
 
 void UMovieSceneSpawnTrack::RemoveSection(UMovieSceneSection& Section)
 {
-	Sections.RemoveAll([&](const UMovieSceneSection* In){ return In == &Section; });
+	Sections.RemoveAll([&](const UMovieSceneSection* In) { return In == &Section; });
+}
+
+void UMovieSceneSpawnTrack::RemoveSectionAt(int32 SectionIndex)
+{
+	Sections.RemoveAt(SectionIndex);
 }
 
 void UMovieSceneSpawnTrack::RemoveAllAnimationData()
diff --git a/Engine/Source/Runtime/MovieSceneTracks/Private/Tracks/MovieSceneStringTrack.cpp b/Engine/Source/Runtime/MovieSceneTracks/Private/Tracks/MovieSceneStringTrack.cpp
--- a/Engine/Source/Runtime/MovieSceneTracks/Private/Tracks/MovieSceneStringTrack.cpp
+++ b/Engine/Source/Runtime/MovieSceneTracks/Private/Tracks/MovieSceneStringTrack.cpp
@@ -64,6 +64,11 @@ void UMovieSceneStringTrack::RemoveSection(UMovieSceneSection& Section)
 }
 
 
+void UMovieSceneStringTrack::RemoveSectionAt(int32 SectionIndex)
+{
+	Sections.RemoveAt(SectionIndex);
+}
+
 #if WITH_EDITORONLY_DATA
 
 FText UMovieSceneStringTrack::GetDefaultDisplayName() const
diff --git a/Engine/Source/Runtime/MovieSceneTracks/Public/Tracks/MovieScene3DConstraintTrack.h b/Engine/Source/Runtime/MovieSceneTracks/Public/Tracks/MovieScene3DConstraintTrack.h
--- a/Engine/Source/Runtime/MovieSceneTracks/Public/Tracks/MovieScene3DConstraintTrack.h
+++ b/Engine/Source/Runtime/MovieSceneTracks/Public/Tracks/MovieScene3DConstraintTrack.h
@@ -38,6 +38,7 @@ public:
 	virtual bool HasSection(const UMovieSceneSection& Section) const override;
 	virtual void AddSection(UMovieSceneSection& Section) override;
 	virtual void RemoveSection(UMovieSceneSection& Section) override;
+	virtual void RemoveSectionAt(int32 SectionIndex) override;
 	virtual bool IsEmpty() const override;
 	virtual const TArray<UMovieSceneSection*>& GetAllSections() const override;
 
diff --git a/Engine/Source/Runtime/MovieSceneTracks/Public/Tracks/MovieSceneAudioTrack.h b/Engine/Source/Runtime/MovieSceneTracks/Public/Tracks/MovieSceneAudioTrack.h
--- a/Engine/Source/Runtime/MovieSceneTracks/Public/Tracks/MovieSceneAudioTrack.h
+++ b/Engine/Source/Runtime/MovieSceneTracks/Public/Tracks/MovieSceneAudioTrack.h
@@ -50,6 +50,7 @@ public:
 	virtual bool HasSection(const UMovieSceneSection& Section) const override;
 	virtual void AddSection(UMovieSceneSection& Section) override;
 	virtual void RemoveSection(UMovieSceneSection& Section) override;
+	virtual void RemoveSectionAt(int32 SectionIndex) override;
 	virtual bool IsEmpty() const override;
 	virtual const TArray<UMovieSceneSection*>& GetAllSections() const override;
 	virtual bool SupportsMultipleRows() const override;
diff --git a/Engine/Source/Runtime/MovieSceneTracks/Public/Tracks/MovieSceneCameraAnimTrack.h b/Engine/Source/Runtime/MovieSceneTracks/Public/Tracks/MovieSceneCameraAnimTrack.h
--- a/Engine/Source/Runtime/MovieSceneTracks/Public/Tracks/MovieSceneCameraAnimTrack.h
+++ b/Engine/Source/Runtime/MovieSceneTracks/Public/Tracks/MovieSceneCameraAnimTrack.h
@@ -30,6 +30,7 @@ public:
 	virtual bool HasSection(const UMovieSceneSection& Section) const override;
 	virtual void AddSection(UMovieSceneSection& Section) override;
 	virtual void RemoveSection(UMovieSceneSection& Section) override;
+	virtual void RemoveSectionAt(int32 SectionIndex) override;
 	virtual bool IsEmpty() const override;
 	virtual bool SupportsMultipleRows() const override { return true; }
 	virtual const TArray<UMovieSceneSection*>& GetAllSections() const override;
diff --git a/Engine/Source/Runtime/MovieSceneTracks/Public/Tracks/MovieSceneCameraCutTrack.h b/Engine/Source/Runtime/MovieSceneTracks/Public/Tracks/MovieSceneCameraCutTrack.h
--- a/Engine/Source/Runtime/MovieSceneTracks/Public/Tracks/MovieSceneCameraCutTrack.h
+++ b/Engine/Source/Runtime/MovieSceneTracks/Public/Tracks/MovieSceneCameraCutTrack.h
@@ -41,6 +41,7 @@ public:
 	virtual UMovieSceneSection* CreateNewSection() override;
 	virtual const TArray<UMovieSceneSection*>& GetAllSections() const override;
 	virtual void RemoveSection(UMovieSceneSection& Section) override;
+	virtual void RemoveSectionAt(int32 SectionIndex) override;
 	virtual void RemoveAllAnimationData() override;
 
 #if WITH_EDITORONLY_DATA
diff --git a/Engine/Source/Runtime/MovieSceneTracks/Public/Tracks/MovieSceneCameraShakeTrack.h b/Engine/Source/Runtime/MovieSceneTracks/Public/Tracks/MovieSceneCameraShakeTrack.h
--- a/Engine/Source/Runtime/MovieSceneTracks/Public/Tracks/MovieSceneCameraShakeTrack.h
+++ b/Engine/Source/Runtime/MovieSceneTracks/Public/Tracks/MovieSceneCameraShakeTrack.h
@@ -30,6 +30,7 @@ public:
 	virtual bool HasSection(const UMovieSceneSection& Section) const override;
 	virtual void AddSection(UMovieSceneSection& Section) override;
 	virtual void RemoveSection(UMovieSceneSection& Section) override;
+	virtual void RemoveSectionAt(int32 SectionIndex) override;
 	virtual bool IsEmpty() const override;
 	virtual bool SupportsMultipleRows() const override { return true; }
 	virtual const TArray<UMovieSceneSection*>& GetAllSections() const override;
diff --git a/Engine/Source/Runtime/MovieSceneTracks/Public/Tracks/MovieSceneCinematicShotTrack.h b/Engine/Source/Runtime/MovieSceneTracks/Public/Tracks/MovieSceneCinematicShotTrack.h
--- a/Engine/Source/Runtime/MovieSceneTracks/Public/Tracks/MovieSceneCinematicShotTrack.h
+++ b/Engine/Source/Runtime/MovieSceneTracks/Public/Tracks/MovieSceneCinematicShotTrack.h
@@ -36,6 +36,7 @@ public:
 	virtual bool SupportsType(TSubclassOf<UMovieSceneSection> SectionClass) const override;
 	virtual UMovieSceneSection* CreateNewSection() override;
 	virtual void RemoveSection(UMovieSceneSection& Section) override;
+	virtual void RemoveSectionAt(int32 SectionIndex) override;
 	virtual bool SupportsMultipleRows() const override;
 	virtual FMovieSceneTrackRowSegmentBlenderPtr GetRowSegmentBlender() const override;
 	virtual FMovieSceneTrackSegmentBlenderPtr GetTrackSegmentBlender() const override;
diff --git a/Engine/Source/Runtime/MovieSceneTracks/Public/Tracks/MovieSceneEventTrack.h b/Engine/Source/Runtime/MovieSceneTracks/Public/Tracks/MovieSceneEventTrack.h
--- a/Engine/Source/Runtime/MovieSceneTracks/Public/Tracks/MovieSceneEventTrack.h
+++ b/Engine/Source/Runtime/MovieSceneTracks/Public/Tracks/MovieSceneEventTrack.h
@@ -57,6 +57,7 @@ public:
 	virtual bool IsEmpty() const override;
 	virtual void RemoveAllAnimationData() override;
 	virtual void RemoveSection(UMovieSceneSection& Section) override;
+	virtual void RemoveSectionAt(int32 SectionIndex) override;
 	virtual FMovieSceneEvalTemplatePtr CreateTemplateForSection(const UMovieSceneSection& InSection) const override;
 	virtual void PostCompile(FMovieSceneEvaluationTrack& Track, const FMovieSceneTrackCompilerArgs& Args) const override;
 	virtual bool SupportsMultipleRows() const override { return true; }
diff --git a/Engine/Source/Runtime/MovieSceneTracks/Public/Tracks/MovieSceneLevelVisibilityTrack.h b/Engine/Source/Runtime/MovieSceneTracks/Public/Tracks/MovieSceneLevelVisibilityTrack.h
--- a/Engine/Source/Runtime/MovieSceneTracks/Public/Tracks/MovieSceneLevelVisibilityTrack.h
+++ b/Engine/Source/Runtime/MovieSceneTracks/Public/Tracks/MovieSceneLevelVisibilityTrack.h
@@ -28,6 +28,7 @@ public:
 	virtual bool IsEmpty() const override;
 	virtual void AddSection(UMovieSceneSection& Section) override;
 	virtual void RemoveSection( UMovieSceneSection& Section ) override;
+	virtual void RemoveSectionAt(int32 SectionIndex) override;
 	virtual bool SupportsType(TSubclassOf<UMovieSceneSection> SectionClass) const override;
 	virtual UMovieSceneSection* CreateNewSection() override;
 	virtual const TArray<UMovieSceneSection*>& GetAllSections() const override;
diff --git a/Engine/Source/Runtime/MovieSceneTracks/Public/Tracks/MovieSceneMaterialTrack.h b/Engine/Source/Runtime/MovieSceneTracks/Public/Tracks/MovieSceneMaterialTrack.h
--- a/Engine/Source/Runtime/MovieSceneTracks/Public/Tracks/MovieSceneMaterialTrack.h
+++ b/Engine/Source/Runtime/MovieSceneTracks/Public/Tracks/MovieSceneMaterialTrack.h
@@ -27,6 +27,7 @@ public:
 	virtual bool HasSection(const UMovieSceneSection& Section) const override;
 	virtual void AddSection(UMovieSceneSection& Section) override;
 	virtual void RemoveSection(UMovieSceneSection& Section) override;
+	virtual void RemoveSectionAt(int32 SectionIndex) override;
 	virtual bool IsEmpty() const override;
 	virtual const TArray<UMovieSceneSection*>& GetAllSections() const override;
 
diff --git a/Engine/Source/Runtime/MovieSceneTracks/Public/Tracks/MovieSceneParticleParameterTrack.h b/Engine/Source/Runtime/MovieSceneTracks/Public/Tracks/MovieSceneParticleParameterTrack.h
--- a/Engine/Source/Runtime/MovieSceneTracks/Public/Tracks/MovieSceneParticleParameterTrack.h
+++ b/Engine/Source/Runtime/MovieSceneTracks/Public/Tracks/MovieSceneParticleParameterTrack.h
@@ -28,6 +28,7 @@ public:
 	virtual bool HasSection(const UMovieSceneSection& Section) const override;
 	virtual void AddSection(UMovieSceneSection& Section) override;
 	virtual void RemoveSection(UMovieSceneSection& Section) override;
+	virtual void RemoveSectionAt(int32 SectionIndex) override;
 	virtual bool IsEmpty() const override;
 	virtual const TArray<UMovieSceneSection*>& GetAllSections() const override;
 
diff --git a/Engine/Source/Runtime/MovieSceneTracks/Public/Tracks/MovieSceneParticleTrack.h b/Engine/Source/Runtime/MovieSceneTracks/Public/Tracks/MovieSceneParticleTrack.h
--- a/Engine/Source/Runtime/MovieSceneTracks/Public/Tracks/MovieSceneParticleTrack.h
+++ b/Engine/Source/Runtime/MovieSceneTracks/Public/Tracks/MovieSceneParticleTrack.h
@@ -36,6 +36,7 @@ public:
 	virtual bool HasSection(const UMovieSceneSection& Section) const override;
 	virtual void AddSection(UMovieSceneSection& Section) override;
 	virtual void RemoveSection(UMovieSceneSection& Section) override;
+	virtual void RemoveSectionAt(int32 SectionIndex) override;
 	virtual bool IsEmpty() const override;
 	virtual const TArray<UMovieSceneSection*>& GetAllSections() const override;
 	virtual void AddNewSection(FFrameNumber SectionTime);
diff --git a/Engine/Source/Runtime/MovieSceneTracks/Public/Tracks/MovieScenePropertyTrack.h b/Engine/Source/Runtime/MovieSceneTracks/Public/Tracks/MovieScenePropertyTrack.h
--- a/Engine/Source/Runtime/MovieSceneTracks/Public/Tracks/MovieScenePropertyTrack.h
+++ b/Engine/Source/Runtime/MovieSceneTracks/Public/Tracks/MovieScenePropertyTrack.h
@@ -24,6 +24,7 @@ public:
 	virtual bool HasSection(const UMovieSceneSection& Section) const override;
 	virtual void AddSection(UMovieSceneSection& Section) override;
 	virtual void RemoveSection(UMovieSceneSection& Section) override;
+	virtual void RemoveSectionAt(int32 SectionIndex) override;
 	virtual bool IsEmpty() const override;
 	virtual const TArray<UMovieSceneSection*>& GetAllSections() const override;
 
diff --git a/Engine/Source/Runtime/MovieSceneTracks/Public/Tracks/MovieSceneSkeletalAnimationTrack.h b/Engine/Source/Runtime/MovieSceneTracks/Public/Tracks/MovieSceneSkeletalAnimationTrack.h
--- a/Engine/Source/Runtime/MovieSceneTracks/Public/Tracks/MovieSceneSkeletalAnimationTrack.h
+++ b/Engine/Source/Runtime/MovieSceneTracks/Public/Tracks/MovieSceneSkeletalAnimationTrack.h
@@ -36,6 +36,7 @@ public:
 	virtual bool HasSection(const UMovieSceneSection& Section) const override;
 	virtual void AddSection(UMovieSceneSection& Section) override;
 	virtual void RemoveSection(UMovieSceneSection& Section) override;
+	virtual void RemoveSectionAt(int32 SectionIndex) override;
 	virtual bool IsEmpty() const override;
 	virtual const TArray<UMovieSceneSection*>& GetAllSections() const override;
 	virtual bool SupportsType(TSubclassOf<UMovieSceneSection> SectionClass) const override;
diff --git a/Engine/Source/Runtime/MovieSceneTracks/Public/Tracks/MovieSceneSpawnTrack.h b/Engine/Source/Runtime/MovieSceneTracks/Public/Tracks/MovieSceneSpawnTrack.h
--- a/Engine/Source/Runtime/MovieSceneTracks/Public/Tracks/MovieSceneSpawnTrack.h
+++ b/Engine/Source/Runtime/MovieSceneTracks/Public/Tracks/MovieSceneSpawnTrack.h
@@ -47,6 +47,7 @@ public:
 	virtual bool HasSection(const UMovieSceneSection& Section) const override;
 	virtual void AddSection(UMovieSceneSection& Section) override;
 	virtual void RemoveSection(UMovieSceneSection& Section) override;
+	virtual void RemoveSectionAt(int32 SectionIndex) override;
 	virtual void RemoveAllAnimationData() override;
 	virtual bool IsEmpty() const override;
 	virtual const TArray<UMovieSceneSection*>& GetAllSections() const override;
diff --git a/Engine/Source/Runtime/MovieSceneTracks/Public/Tracks/MovieSceneStringTrack.h b/Engine/Source/Runtime/MovieSceneTracks/Public/Tracks/MovieSceneStringTrack.h
--- a/Engine/Source/Runtime/MovieSceneTracks/Public/Tracks/MovieSceneStringTrack.h
+++ b/Engine/Source/Runtime/MovieSceneTracks/Public/Tracks/MovieSceneStringTrack.h
@@ -39,6 +39,7 @@ public:
 	virtual bool IsEmpty() const override;
 	virtual void RemoveAllAnimationData() override;
 	virtual void RemoveSection(UMovieSceneSection& Section) override;
+	virtual void RemoveSectionAt(int32 SectionIndex) override;
 
 #if WITH_EDITORONLY_DATA
 	virtual FText GetDefaultDisplayName() const override;
