- Fix crash when picking a color from Composure
Integration of CL: 7445029

diff --git a/Engine/Plugins/Compositing/Composure/Source/ComposureLayersEditor/Private/ComposureDetailCustomizations.cpp b/Engine/Plugins/Compositing/Composure/Source/ComposureLayersEditor/Private/ComposureDetailCustomizations.cpp
--- a/Engine/Plugins/Compositing/Composure/Source/ComposureLayersEditor/Private/ComposureDetailCustomizations.cpp
+++ b/Engine/Plugins/Compositing/Composure/Source/ComposureLayersEditor/Private/ComposureDetailCustomizations.cpp
@@ -708,11 +708,11 @@ void FCompositingMaterialPassCustomization::CustomizeChildren(TSharedRef<IProper
 
 					RebuildTextureSourceList();
 
-					FIsResetToDefaultVisible IsResetVisible = FIsResetToDefaultVisible::CreateRaw(this, &FCompositingMaterialPassCustomization::TextureShouldShowResetToDefault);
-					FResetToDefaultHandler ResetHandler = FResetToDefaultHandler::CreateRaw(this, &FCompositingMaterialPassCustomization::TextureResetToDefault, PassComboButton);
+					FIsResetToDefaultVisible IsResetVisible = FIsResetToDefaultVisible::CreateSP(this, &FCompositingMaterialPassCustomization::TextureShouldShowResetToDefault);
+					FResetToDefaultHandler ResetHandler = FResetToDefaultHandler::CreateSP(this, &FCompositingMaterialPassCustomization::TextureResetToDefault, PassComboButton);
 					FResetToDefaultOverride ResetOverride = FResetToDefaultOverride::Create(IsResetVisible, ResetHandler);
 
-					//ChildHandle->SetOnPropertyValueChanged(FSimpleDelegate::CreateRaw(this, &FCompositingMaterialPassCustomization::OnTextureOverrideChanged, ChildHandle));
+					//ChildHandle->SetOnPropertyValueChanged(FSimpleDelegate::CreateSP(this, &FCompositingMaterialPassCustomization::OnTextureOverrideChanged, ChildHandle));
 
 					PropertyRow.OverrideResetToDefault(ResetOverride);
 					PropertyRow.GetDefaultWidgets(NameWidget, ValueWidget, Row);
@@ -848,11 +848,11 @@ void FCompositingMaterialPassCustomization::CustomizeChildren(TSharedRef<IProper
 					}
 				}
 
-				FIsResetToDefaultVisible IsResetVisible = FIsResetToDefaultVisible::CreateRaw(this, &FCompositingMaterialPassCustomization::VectorShouldShowResetToDefault);
-				FResetToDefaultHandler ResetHandler = FResetToDefaultHandler::CreateRaw(this, &FCompositingMaterialPassCustomization::VectorResetToDefault);
+				FIsResetToDefaultVisible IsResetVisible = FIsResetToDefaultVisible::CreateSP(this, &FCompositingMaterialPassCustomization::VectorShouldShowResetToDefault);
+				FResetToDefaultHandler ResetHandler = FResetToDefaultHandler::CreateSP(this, &FCompositingMaterialPassCustomization::VectorResetToDefault);
 				FResetToDefaultOverride ResetOverride = FResetToDefaultOverride::Create(IsResetVisible, ResetHandler);
 
-				ChildHandle->SetOnPropertyValueChanged(FSimpleDelegate::CreateRaw(this, &FCompositingMaterialPassCustomization::OnVectorOverrideChanged, ChildHandle));
+				ChildHandle->SetOnPropertyValueChanged(FSimpleDelegate::CreateSP(this, &FCompositingMaterialPassCustomization::OnVectorOverrideChanged, ChildHandle));
 
 				TSharedPtr<SWidget> NameWidget;
 				TSharedPtr<SWidget> ValueWidget;
