Fix Performance warning (Composure)

diff --git a/Engine/Plugins/Compositing/Composure/Source/ComposureLayersEditor/Private/CompElementEditorModule.cpp b/Engine/Plugins/Compositing/Composure/Source/ComposureLayersEditor/Private/CompElementEditorModule.cpp
--- a/Engine/Plugins/Compositing/Composure/Source/ComposureLayersEditor/Private/CompElementEditorModule.cpp
+++ b/Engine/Plugins/Compositing/Composure/Source/ComposureLayersEditor/Private/CompElementEditorModule.cpp
@@ -233,7 +233,7 @@ TSharedPtr<SWindow> FCompElementEditorModule::RequestCompositingPickerWindow(TWe
 //------------------------------------------------------------------------------
 bool FCompElementEditorModule::DeferCompositingDraw(ACompositingElement* CompElement)
 {
-	IConsoleVariable* CVarUsingDecoupledDrawing = IConsoleManager::Get().FindConsoleVariable(TEXT("r.Composure.CompositingElements.Editor.DecoupleRenderingFromLevelViewport"));
+	static IConsoleVariable* CVarUsingDecoupledDrawing = IConsoleManager::Get().FindConsoleVariable(TEXT("r.Composure.CompositingElements.Editor.DecoupleRenderingFromLevelViewport"));
 	if (CompElement && CVarUsingDecoupledDrawing && CVarUsingDecoupledDrawing->GetInt() > 0 && !CompElementManager->IsDrawing(CompElement))
 	{
 		UWorld* World = CompElement->GetWorld();
