UK2Node_BaseAsyncTask: Rollback change 4421805 (custom event node marked as CallInEditor)
Adding FEditorScriptExecutionGuard for OnPick & OnCancel delegates for CompositingPickerAsyncTask

As discussed with gameplay team, another approach should be used for composure instead of marking output delegates of async task node as CallInEditor. If the feature is still required they can use FEditorScriptExecutionGuard around the delegate broadcast calls for the nodes that are expected to work outside of PIE.

#rnx

diff --git a/Engine/Plugins/Compositing/Composure/Source/Composure/Private/EditorSupport/CompositingEditorSupportLibrary.cpp b/Engine/Plugins/Compositing/Composure/Source/Composure/Private/EditorSupport/CompositingEditorSupportLibrary.cpp
--- a/Engine/Plugins/Compositing/Composure/Source/Composure/Private/EditorSupport/CompositingEditorSupportLibrary.cpp
+++ b/Engine/Plugins/Compositing/Composure/Source/Composure/Private/EditorSupport/CompositingEditorSupportLibrary.cpp
@@ -7,6 +7,7 @@
 
 #include "EditorSupport/ICompositingEditor.h"
 #include "EditorSupport/WeakUInterfacePtr.h"
+#include "UObject/Script.h" // for FEditorScriptExecutionGuard
 
 UCompositingPickerAsyncTask* UCompositingPickerAsyncTask::OpenCompositingPicker(UTextureRenderTarget2D* PickerTarget, UTexture* DisplayImage, FText WindowTitle, const bool bAverageColorOnDrag, const bool bUseImplicitGamma)
 {
@@ -67,6 +68,8 @@ void UCompositingPickerAsyncTask::Open(UTextureRenderTarget2D* InPickerTarget, U
 
 void UCompositingPickerAsyncTask::InternalOnPick(const FVector2D& PickedUV, const FLinearColor& PickedColor, bool bInteractive)
 {
+	FEditorScriptExecutionGuard EdScriptGuard;
+
 	if (bInteractive)
 	{
 		OnPick.Broadcast(PickedUV, PickedColor);
@@ -80,6 +83,8 @@ void UCompositingPickerAsyncTask::InternalOnPick(const FVector2D& PickedUV, cons
 
 void UCompositingPickerAsyncTask::InternalOnCancel()
 {
+	FEditorScriptExecutionGuard EdScriptGuard;
+
 	OnCancel.Broadcast(FVector2D(-1.f, -1.f), FLinearColor::Black);
 	SetReadyToDestroy();
 }
diff --git a/Engine/Source/Editor/BlueprintGraph/Private/K2Node_BaseAsyncTask.cpp b/Engine/Source/Editor/BlueprintGraph/Private/K2Node_BaseAsyncTask.cpp
--- a/Engine/Source/Editor/BlueprintGraph/Private/K2Node_BaseAsyncTask.cpp
+++ b/Engine/Source/Editor/BlueprintGraph/Private/K2Node_BaseAsyncTask.cpp
@@ -286,7 +286,6 @@ bool UK2Node_BaseAsyncTask::FBaseAsyncTaskHelper::HandleDelegateImplementation(
 		InOutLastThenPin = AddDelegateNode->FindPinChecked(UEdGraphSchema_K2::PN_Then);
 		CurrentCENode->CustomFunctionName = *FString::Printf(TEXT("%s_%s"), *CurrentProperty->GetName(), *CompilerContext.GetGuid(CurrentNode));
 		CurrentCENode->AllocateDefaultPins();
-		CurrentCENode->bCallInEditor = true;
 
 		bIsErrorFree &= FBaseAsyncTaskHelper::CreateDelegateForNewFunction(AddDelegateNode->GetDelegatePin(), CurrentCENode->GetFunctionName(), CurrentNode, SourceGraph, CompilerContext);
 		bIsErrorFree &= FBaseAsyncTaskHelper::CopyEventSignature(CurrentCENode, AddDelegateNode->GetDelegateSignature(), Schema);
