Composure change to support color grading and white balance via floating point post processing LUT

diff --git a/Engine/Plugins/Compositing/Composure/Source/Composure/Private/ComposureBlueprintLibrary.cpp b/Engine/Plugins/Compositing/Composure/Source/Composure/Private/ComposureBlueprintLibrary.cpp
--- a/Engine/Plugins/Compositing/Composure/Source/Composure/Private/ComposureBlueprintLibrary.cpp
+++ b/Engine/Plugins/Compositing/Composure/Source/Composure/Private/ComposureBlueprintLibrary.cpp
@@ -84,51 +84,12 @@ void UComposureBlueprintLibrary::CopyCameraSettingsToSceneCapture(UCameraCompone
 		const FPostProcessSettings& SrcPPSettings = CameraViewInfo.PostProcessSettings;
 		FPostProcessSettings& DstPPSettings = Dst->PostProcessSettings;
 
-#define COPY_PPSETTING(FieldName, ToggleName) \
-		DstPPSettings.FieldName  = SrcPPSettings.FieldName; \
-		DstPPSettings.ToggleName = true;
-
-		COPY_PPSETTING(WhiteTemp, bOverride_WhiteTemp);
-		
-		COPY_PPSETTING(ColorSaturation, bOverride_ColorSaturation);
-		COPY_PPSETTING(ColorContrast, bOverride_ColorContrast);
-		COPY_PPSETTING(ColorGamma, bOverride_ColorGamma);
-		COPY_PPSETTING(ColorGain, bOverride_ColorGain);
-		COPY_PPSETTING(ColorOffset, bOverride_ColorOffset);
-
-		COPY_PPSETTING(ColorSaturationShadows, bOverride_ColorSaturationShadows);
-		COPY_PPSETTING(ColorContrastShadows, bOverride_ColorContrastShadows);
-		COPY_PPSETTING(ColorGammaShadows, bOverride_ColorGammaShadows);
-		COPY_PPSETTING(ColorGainShadows, bOverride_ColorGainShadows);
-		COPY_PPSETTING(ColorOffsetShadows, bOverride_ColorOffsetShadows);
-
-		COPY_PPSETTING(ColorCorrectionShadowsMax, bOverride_ColorCorrectionShadowsMax);
-		
-		COPY_PPSETTING(ColorSaturationMidtones, bOverride_ColorSaturationMidtones);
-		COPY_PPSETTING(ColorContrastMidtones, bOverride_ColorContrastMidtones);
-		COPY_PPSETTING(ColorGammaMidtones, bOverride_ColorGammaMidtones);
-		COPY_PPSETTING(ColorGainMidtones, bOverride_ColorGainMidtones);
-		COPY_PPSETTING(ColorOffsetMidtones, bOverride_ColorOffsetMidtones);
-		
-		COPY_PPSETTING(ColorSaturationHighlights, bOverride_ColorSaturationHighlights);
-		COPY_PPSETTING(ColorContrastHighlights, bOverride_ColorContrastHighlights);
-		COPY_PPSETTING(ColorGammaHighlights, bOverride_ColorGammaHighlights);
-		COPY_PPSETTING(ColorGainHighlights, bOverride_ColorGainHighlights);
-		COPY_PPSETTING(ColorOffsetHighlights, bOverride_ColorOffsetHighlights);
-
-		COPY_PPSETTING(ColorCorrectionHighlightsMin, bOverride_ColorCorrectionHighlightsMin);
-			
-		COPY_PPSETTING(DepthOfFieldFstop, bOverride_DepthOfFieldFstop);
-		COPY_PPSETTING(DepthOfFieldMinFstop, bOverride_DepthOfFieldMinFstop);
-		COPY_PPSETTING(DepthOfFieldBladeCount, bOverride_DepthOfFieldBladeCount);
-		COPY_PPSETTING(AutoExposureBias, bOverride_AutoExposureBias);
-		COPY_PPSETTING(AutoExposureBiasCurve, bOverride_AutoExposureBiasCurve);
-			
-		COPY_PPSETTING(DepthOfFieldSensorWidth, bOverride_DepthOfFieldSensorWidth);
-		COPY_PPSETTING(DepthOfFieldFocalDistance, bOverride_DepthOfFieldFocalDistance);
-		COPY_PPSETTING(DepthOfFieldDepthBlurAmount, bOverride_DepthOfFieldDepthBlurAmount);
-		COPY_PPSETTING(DepthOfFieldDepthBlurRadius, bOverride_DepthOfFieldDepthBlurRadius);
-
-#undef COPY_PPSETTING
+		FWeightedBlendables DstWeightedBlendables = DstPPSettings.WeightedBlendables;
+
+		// Copy all of the post processing settings
+		DstPPSettings = SrcPPSettings;
+
+		// But restore the original blendables
+		DstPPSettings.WeightedBlendables = DstWeightedBlendables;
 	}
 }
diff --git a/Engine/Shaders/Private/PostProcessCombineLUTs.usf b/Engine/Shaders/Private/PostProcessCombineLUTs.usf
--- a/Engine/Shaders/Private/PostProcessCombineLUTs.usf
+++ b/Engine/Shaders/Private/PostProcessCombineLUTs.usf
@@ -299,7 +299,7 @@ float4 CombineLUTsCommon(float2 InUV, uint InLayerIndex)
 
 	ColorAP1 = ColorCorrectAll( ColorAP1 );
 
-	// Store for Legacy tonemap later
+	// Store for Legacy tonemap later and for Linear HDR output without tone curve
 	float3 GradedColor = mul( AP1_2_sRGB, ColorAP1 );
 
 	const float3x3 BlueCorrect =
@@ -440,6 +440,11 @@ float4 CombineLUTsCommon(float2 InUV, uint InLayerIndex)
 			float3 OutputGamutColor = mul( AP1_2_Output, mul( sRGB_2_AP1, GradedColor ) );
 			OutDeviceColor = LinearToST2084( OutputGamutColor );
 	}
+	// Linear HDR, including all color correction, but no tone curve
+	else if( GetOutputDevice() == 8 ) 
+	{
+			OutDeviceColor = GradedColor;
+	}
 
 	// OutputDevice == 2
 	// Gamma 2.2, user specified gamut
diff --git a/Engine/Shaders/Private/PostProcessTonemap.usf b/Engine/Shaders/Private/PostProcessTonemap.usf
--- a/Engine/Shaders/Private/PostProcessTonemap.usf
+++ b/Engine/Shaders/Private/PostProcessTonemap.usf
@@ -21,6 +21,7 @@
 	#define TONEMAPPER_OUTOUT_ACES1000nitScRGB 5
 	#define TONEMAPPER_OUTOUT_ACES2000nitScRGB 6
 	#define TONEMAPPER_OUTOUT_LinearEXR 7
+	#define TONEMAPPER_OUTOUT_NoToneCurve 8
 
 
 #ifndef DIM_OUTPUT_DEVICE
@@ -76,7 +77,7 @@ half3 ColorLookupTable( half3 LinearColor )
 {
 	float3 LUTEncodedColor;
 	// Encode as ST-2084 (Dolby PQ) values
-	#if (DIM_OUTPUT_DEVICE == TONEMAPPER_OUTOUT_ACES1000nitST2084 || DIM_OUTPUT_DEVICE == TONEMAPPER_OUTOUT_ACES2000nitST2084 || DIM_OUTPUT_DEVICE == TONEMAPPER_OUTOUT_ACES1000nitScRGB || DIM_OUTPUT_DEVICE == TONEMAPPER_OUTOUT_ACES2000nitScRGB || DIM_OUTPUT_DEVICE == TONEMAPPER_OUTOUT_LinearEXR)
+	#if (DIM_OUTPUT_DEVICE == TONEMAPPER_OUTOUT_ACES1000nitST2084 || DIM_OUTPUT_DEVICE == TONEMAPPER_OUTOUT_ACES2000nitST2084 || DIM_OUTPUT_DEVICE == TONEMAPPER_OUTOUT_ACES1000nitScRGB || DIM_OUTPUT_DEVICE == TONEMAPPER_OUTOUT_ACES2000nitScRGB || DIM_OUTPUT_DEVICE == TONEMAPPER_OUTOUT_LinearEXR || DIM_OUTPUT_DEVICE == TONEMAPPER_OUTOUT_NoToneCurve)
 		// ST2084 expects to receive linear values 0-10000 in nits.
 		// So the linear value must be multiplied by a scale factor to convert to nits.
 		LUTEncodedColor = LinearToST2084(LinearColor * LinearToNitsScale);
@@ -509,7 +510,7 @@ float4 TonemapCommonPS(
 	
 		half LuminanceForPostProcessAA  = dot(OutDeviceColor, half3(0.299f, 0.587f, 0.114f));
 
-		#if USE_GRAIN_QUANTIZATION
+		#if USE_GRAIN_QUANTIZATION && DIM_OUTPUT_DEVICE != TONEMAPPER_OUTOUT_NoToneCurve
 			// Needs to go after tonemapping.
 			#if USE_LINEAR_FLOAT_RT
 				half GrainQuantization = 1.0/1024.0;
@@ -533,12 +534,13 @@ float4 TonemapCommonPS(
 
 		#elif DIM_OUTPUT_DEVICE == TONEMAPPER_OUTOUT_LinearEXR
 			OutColor.xyz = ST2084ToLinear(OutColor.xyz);
-
+		#elif DIM_OUTPUT_DEVICE == TONEMAPPER_OUTOUT_NoToneCurve
+			OutColor.xyz = OutDeviceColor;
 		#endif
 
 	#endif
 	
-	#if POST_PROCESS_ALPHA == 2
+	#if POST_PROCESS_ALPHA == 2 || (POST_PROCESS_ALPHA == 1 && DIM_OUTPUT_DEVICE == TONEMAPPER_OUTOUT_NoToneCurve)
 		OutColor.a = SceneColor.a;
 	#endif
 
diff --git a/Engine/Source/Runtime/Engine/Classes/Engine/EngineTypes.h b/Engine/Source/Runtime/Engine/Classes/Engine/EngineTypes.h
--- a/Engine/Source/Runtime/Engine/Classes/Engine/EngineTypes.h
+++ b/Engine/Source/Runtime/Engine/Classes/Engine/EngineTypes.h
@@ -355,7 +355,8 @@ enum ESceneCaptureSource
 	SCS_SceneDepth UMETA(DisplayName="SceneDepth in R"),
 	SCS_DeviceDepth UMETA(DisplayName = "DeviceDepth in RGB"),
 	SCS_Normal UMETA(DisplayName="Normal in RGB (Deferred Renderer only)"),
-	SCS_BaseColor UMETA(DisplayName="BaseColor in RGB (Deferred Renderer only)")
+	SCS_BaseColor UMETA(DisplayName = "BaseColor in RGB (Deferred Renderer only)"),
+	SCS_FinalColorHDR UMETA(DisplayName = "Final Color (HDR) in Linear sRGB gamut")
 };
 
 /** Specifies how scene captures are composited into render buffers */
diff --git a/Engine/Source/Runtime/Engine/Private/ShaderCompiler/ShaderCompiler.cpp b/Engine/Source/Runtime/Engine/Private/ShaderCompiler/ShaderCompiler.cpp
--- a/Engine/Source/Runtime/Engine/Private/ShaderCompiler/ShaderCompiler.cpp
+++ b/Engine/Source/Runtime/Engine/Private/ShaderCompiler/ShaderCompiler.cpp
@@ -3872,7 +3872,7 @@ void VerifyGlobalShaders(EShaderPlatform Platform, bool bLoadedFromCacheFile)
 		}
 
 		ensureMsgf(
-			PermutationCountToCompile < 353,	// ToneMapper today can go up to 353 permutations
+			PermutationCountToCompile < 397,	// ToneMapper today (2019-04-17) can go up to 396 permutations
 			TEXT("Global shader %s has %i permutation: probably more that it needs."),
 			GlobalShaderType->GetName(), PermutationCountToCompile);
 
diff --git a/Engine/Source/Runtime/Renderer/Private/PostProcess/PostProcessCombineLUTs.cpp b/Engine/Source/Runtime/Renderer/Private/PostProcess/PostProcessCombineLUTs.cpp
--- a/Engine/Source/Runtime/Renderer/Private/PostProcess/PostProcessCombineLUTs.cpp
+++ b/Engine/Source/Runtime/Renderer/Private/PostProcess/PostProcessCombineLUTs.cpp
@@ -251,7 +251,17 @@ public:
 			static TConsoleVariableData<int32>* CVarOutputDevice = IConsoleManager::Get().FindTConsoleVariableDataInt(TEXT("r.HDR.Display.OutputDevice"));
 			static TConsoleVariableData<float>* CVarOutputGamma = IConsoleManager::Get().FindTConsoleVariableDataFloat(TEXT("r.TonemapperGamma"));
 
-			int32 OutputDeviceValue = CVarOutputDevice->GetValueOnRenderThread();
+			int32 OutputDeviceValue;
+
+			if (ViewFamily.SceneCaptureSource == SCS_FinalColorHDR)
+			{
+				OutputDeviceValue = 8; //LinearNoToneCurve from FTonemapperOutputDevice
+			}
+			else
+			{
+				OutputDeviceValue = CVarOutputDevice->GetValueOnRenderThread();
+			}
+
 			float Gamma = CVarOutputGamma->GetValueOnRenderThread();
 
 			if (PLATFORM_APPLE && Gamma == 0.0f)
@@ -886,7 +896,7 @@ void FRCPassPostProcessCombineLUTs::Process(FRenderingCompositePassContext& Cont
 	// and the result is reused by eSSP_RIGHT_EYE.   Eye-adaptation for stereo works in a similar way.
 	// Fundamentally, this relies on the fact that the view is recycled when doing stereo rendering and the LEFT eye is done first.
 	const FSceneRenderTargetItem* DestRenderTarget = !bAllocateOutput ?
-		Context.View.GetTonemappingLUTRenderTarget(Context.RHICmdList, GLUTSize, bUseVolumeTextureLUT, bIsComputePass) :
+		Context.View.GetTonemappingLUTRenderTarget(Context.RHICmdList, GLUTSize, bUseVolumeTextureLUT, bIsComputePass, ViewFamily.SceneCaptureSource == SCS_FinalColorHDR) :
 		&PassOutputs[0].RequestSurface(Context);
 	
 	check(DestRenderTarget);
@@ -1005,21 +1015,7 @@ FPooledRenderTargetDesc FRCPassPostProcessCombineLUTs::ComputeOutputDesc(EPassOu
 	}
 	else
 	{
-		EPixelFormat LUTPixelFormat = PF_A2B10G10R10;
-		if (!GPixelFormats[LUTPixelFormat].Supported)
-		{
-			LUTPixelFormat = PF_R8G8B8A8;
-		}
-		
-		Ret = FPooledRenderTargetDesc::Create2DDesc(FIntPoint(GLUTSize * GLUTSize, GLUTSize), LUTPixelFormat, FClearValueBinding::Transparent, TexCreate_None, TexCreate_RenderTargetable | TexCreate_ShaderResource, false);
-		
-		if(RuntimeVolumeTextureLUTSupported(ShaderPlatform))
-		{
-			Ret.Extent = FIntPoint(GLUTSize, GLUTSize);
-			Ret.Depth = GLUTSize;
-		}
-		Ret.Flags |= GFastVRamConfig.CombineLUTs;
-		Ret.DebugName = TEXT("CombineLUTs");
+		Ret =FSceneViewState::CreateLUTRenderTarget(GLUTSize, RuntimeVolumeTextureLUTSupported(ShaderPlatform), false, bNeedFloatOutput);
 	}
 	Ret.ClearValue = FClearValueBinding::Transparent;
 
diff --git a/Engine/Source/Runtime/Renderer/Private/PostProcess/PostProcessCombineLUTs.h b/Engine/Source/Runtime/Renderer/Private/PostProcess/PostProcessCombineLUTs.h
--- a/Engine/Source/Runtime/Renderer/Private/PostProcess/PostProcessCombineLUTs.h
+++ b/Engine/Source/Runtime/Renderer/Private/PostProcess/PostProcessCombineLUTs.h
@@ -21,9 +21,10 @@ bool RuntimeVolumeTextureLUTSupported(EShaderPlatform Platform);
 class FRCPassPostProcessCombineLUTs : public TRenderingCompositePassBase<0, 1>
 {
 public:
-	FRCPassPostProcessCombineLUTs(EShaderPlatform InShaderPlatform, bool bInAllocateOutput, bool InIsComputePass)
+	FRCPassPostProcessCombineLUTs(EShaderPlatform InShaderPlatform, bool bInAllocateOutput, bool InIsComputePass, bool bInNeedFloatOutput)
 	: ShaderPlatform(InShaderPlatform)
 	, bAllocateOutput(bInAllocateOutput)
+	, bNeedFloatOutput(bInNeedFloatOutput)
 	{
 		bIsComputePass = InIsComputePass;
 		bPreferAsyncCompute = false;
@@ -49,6 +50,7 @@ private:
 
 	EShaderPlatform ShaderPlatform;
 	bool bAllocateOutput;
+	bool bNeedFloatOutput;
 };
 
 
diff --git a/Engine/Source/Runtime/Renderer/Private/PostProcess/PostProcessTonemap.cpp b/Engine/Source/Runtime/Renderer/Private/PostProcess/PostProcessTonemap.cpp
--- a/Engine/Source/Runtime/Renderer/Private/PostProcess/PostProcessTonemap.cpp
+++ b/Engine/Source/Runtime/Renderer/Private/PostProcess/PostProcessTonemap.cpp
@@ -60,6 +60,7 @@ enum class FTonemapperOutputDevice
 	ACES1000nitScRGB,
 	ACES2000nitScRGB,
 	LinearEXR,
+	LinearNoToneCurve,
 
 	MAX
 };
@@ -75,7 +76,8 @@ static TAutoConsoleVariable<int32> CVarDisplayOutputDevice(
 	TEXT("4: ACES 2000 nit ST-2084 (Dolby PQ) (HDR)\n")
 	TEXT("5: ACES 1000 nit ScRGB (HDR)\n")
 	TEXT("6: ACES 2000 nit ScRGB (HDR)\n")
-	TEXT("7: Linear EXR (HDR)\n"),
+	TEXT("7: Linear EXR (HDR)\n")
+	TEXT("8: Linear final color, no tone curve (HDR)\n"),
 	ECVF_Scalability | ECVF_RenderThreadSafe);
 	
 static TAutoConsoleVariable<int32> CVarHDROutputEnabled(
@@ -994,7 +996,14 @@ void FRCPassPostProcessTonemap::Process(FRenderingCompositePassContext& Context)
 			DesktopPermutationVector.Set<TonemapperPermutation::FTonemapperColorFringeDim>(View.FinalPostProcessSettings.SceneFringeIntensity > 0.01f);
 		}
 
-		DesktopPermutationVector.Set<TonemapperPermutation::FTonemapperOutputDeviceDim>(GetOutputDeviceValue());
+		if (ViewFamily.SceneCaptureSource == SCS_FinalColorHDR)
+		{
+			DesktopPermutationVector.Set<TonemapperPermutation::FTonemapperOutputDeviceDim>(FTonemapperOutputDevice::LinearNoToneCurve);
+		}
+		else
+		{
+			DesktopPermutationVector.Set<TonemapperPermutation::FTonemapperOutputDeviceDim>(GetOutputDeviceValue());
+		}
 
 		DesktopPermutationVector = TonemapperPermutation::RemapPermutation(DesktopPermutationVector);
 	}
@@ -1170,6 +1179,10 @@ FPooledRenderTargetDesc FRCPassPostProcessTonemap::ComputeOutputDesc(EPassOutput
 	{
 		Ret.Format = PF_A32B32G32R32F;
 	}
+	if (CVarDisplayOutputDevice.GetValueOnRenderThread() == 8)
+	{
+		Ret.Format = PF_FloatRGBA;
+	}
 
 
 	// Mobile needs to override the extent
diff --git a/Engine/Source/Runtime/Renderer/Private/PostProcess/PostProcessing.cpp b/Engine/Source/Runtime/Renderer/Private/PostProcess/PostProcessing.cpp
--- a/Engine/Source/Runtime/Renderer/Private/PostProcess/PostProcessing.cpp
+++ b/Engine/Source/Runtime/Renderer/Private/PostProcess/PostProcessing.cpp
@@ -378,7 +378,10 @@ static FRCPassPostProcessTonemap* AddTonemapper(
 	FRenderingCompositeOutputRef TonemapperCombinedLUTOutputRef;
 	if (StereoPass != eSSP_RIGHT_EYE)
 	{
-		FRenderingCompositePass* CombinedLUT = Context.Graph.RegisterPass(new(FMemStack::Get()) FRCPassPostProcessCombineLUTs(View.GetShaderPlatform(), View.State == nullptr, bIsComputePass));
+		bool bNeedFloatOutput = View.Family->SceneCaptureSource == SCS_FinalColorHDR;
+		bool bAllocateOutput = View.State == nullptr;
+
+		FRenderingCompositePass* CombinedLUT = Context.Graph.RegisterPass(new(FMemStack::Get()) FRCPassPostProcessCombineLUTs(View.GetShaderPlatform(), bAllocateOutput, bIsComputePass, bNeedFloatOutput));
 		TonemapperCombinedLUTOutputRef =  FRenderingCompositeOutputRef(CombinedLUT);
 	}
 
@@ -1365,7 +1368,7 @@ void FPostProcessing::Process(FRHICommandListImmediate& RHICmdList, const FViewI
 		const bool bHDROutputEnabled = GRHISupportsHDROutput && IsHDREnabled();
 
 		static const auto CVarDumpFramesAsHDR = IConsoleManager::Get().FindTConsoleVariableDataInt(TEXT("r.BufferVisualizationDumpFramesAsHDR"));
-		const bool bHDRTonemapperOutput = bAllowTonemapper && (GetHighResScreenshotConfig().bCaptureHDR || CVarDumpFramesAsHDR->GetValueOnRenderThread() || bHDROutputEnabled);
+		const bool bHDRTonemapperOutput = bAllowTonemapper && (View.Family->SceneCaptureSource == SCS_FinalColorHDR || GetHighResScreenshotConfig().bCaptureHDR || CVarDumpFramesAsHDR->GetValueOnRenderThread() || bHDROutputEnabled);
 
 		FRCPassPostProcessTonemap* Tonemapper = 0;
 
@@ -2052,6 +2055,7 @@ void FPostProcessing::Process(FRHICommandListImmediate& RHICmdList, const FViewI
 
 			// todo: this should come from View.Family->RenderTarget
 			Desc.Format = bHDROutputEnabled ? GRHIHDRDisplayOutputFormat : Desc.Format;
+			if (View.Family->SceneCaptureSource == SCS_FinalColorHDR) Desc.Format = PF_FloatRGBA;
 			Desc.NumMips = 1;
 			Desc.DebugName = TEXT("FinalPostProcessColor");
 
diff --git a/Engine/Source/Runtime/Renderer/Private/SceneCaptureRendering.cpp b/Engine/Source/Runtime/Renderer/Private/SceneCaptureRendering.cpp
--- a/Engine/Source/Runtime/Renderer/Private/SceneCaptureRendering.cpp
+++ b/Engine/Source/Runtime/Renderer/Private/SceneCaptureRendering.cpp
@@ -51,7 +51,8 @@ const TCHAR* GShaderSourceModeDefineName[] =
 	TEXT("SOURCE_MODE_SCENE_DEPTH"),
 	TEXT("SOURCE_MODE_DEVICE_DEPTH"),
 	TEXT("SOURCE_MODE_NORMAL"),
-	TEXT("SOURCE_MODE_BASE_COLOR")
+	TEXT("SOURCE_MODE_BASE_COLOR"),
+	nullptr
 };
 
 static TAutoConsoleVariable<int32> CVarEnableViewExtensionsForSceneCapture(
@@ -196,7 +197,7 @@ void FDeferredShadingSceneRenderer::CopySceneCaptureComponentToTarget(FRHIComman
 		SceneCaptureSource = SCS_SceneColorHDR;
 	}
 
-	if (SceneCaptureSource != SCS_FinalColorLDR)
+	if (SceneCaptureSource != SCS_FinalColorLDR && SceneCaptureSource != SCS_FinalColorHDR)
 	{
 		SCOPED_DRAW_EVENT(RHICmdList, CaptureSceneComponent);
 
@@ -709,7 +710,8 @@ void FScene::UpdateSceneCaptureContents(USceneCaptureComponent2D* CaptureCompone
 			BuildProjectionMatrix(CaptureSize, CaptureComponent->ProjectionType, FOV, CaptureComponent->OrthoWidth, ProjectionMatrix);
 		}
 
-		const bool bUseSceneColorTexture = CaptureComponent->CaptureSource != SCS_FinalColorLDR;
+		const bool bUseSceneColorTexture = CaptureComponent->CaptureSource != SCS_FinalColorLDR &&
+			CaptureComponent->CaptureSource != SCS_FinalColorHDR;
 
 		FSceneRenderer* SceneRenderer = CreateSceneRendererForSceneCapture(
 			this, 
diff --git a/Engine/Source/Runtime/Renderer/Private/ScenePrivate.h b/Engine/Source/Runtime/Renderer/Private/ScenePrivate.h
--- a/Engine/Source/Runtime/Renderer/Private/ScenePrivate.h
+++ b/Engine/Source/Runtime/Renderer/Private/ScenePrivate.h
@@ -1167,9 +1167,36 @@ public:
 		bValidTonemappingLUT = bValid;
 	}
 
+	static FPooledRenderTargetDesc CreateLUTRenderTarget(const int32 LUTSize, const bool bUseVolumeLUT, const bool bNeedUAV, const bool bNeedFloatOutput)
+	{
+		// Create the texture needed for the tonemapping LUT in one place
+		EPixelFormat LUTPixelFormat = PF_A2B10G10R10;
+		if (!GPixelFormats[LUTPixelFormat].Supported)
+		{
+			LUTPixelFormat = PF_R8G8B8A8;
+		}
+		if (bNeedFloatOutput)
+		{
+			LUTPixelFormat = PF_FloatRGBA;
+		}
+
+		FPooledRenderTargetDesc Desc = FPooledRenderTargetDesc::Create2DDesc(FIntPoint(LUTSize * LUTSize, LUTSize), LUTPixelFormat, FClearValueBinding::Transparent, TexCreate_None, TexCreate_ShaderResource, false);
+		Desc.TargetableFlags |= bNeedUAV ? TexCreate_UAV : TexCreate_RenderTargetable;
+
+		if (bUseVolumeLUT)
+		{
+			Desc.Extent = FIntPoint(LUTSize, LUTSize);
+			Desc.Depth = LUTSize;
+		}
+
+		Desc.DebugName = TEXT("CombineLUTs");
+		Desc.Flags |= GFastVRamConfig.CombineLUTs;
+
+		return Desc;
+	}
 
 	// Returns a reference to the render target used for the LUT.  Allocated on the first request.
-	FSceneRenderTargetItem& GetTonemappingLUTRenderTarget(FRHICommandList& RHICmdList, const int32 LUTSize, const bool bUseVolumeLUT, const bool bNeedUAV)
+	FSceneRenderTargetItem& GetTonemappingLUTRenderTarget(FRHICommandList& RHICmdList, const int32 LUTSize, const bool bUseVolumeLUT, const bool bNeedUAV, const bool bNeedFloatOutput)
 	{
 		if (CombinedLUTRenderTarget.IsValid() == false || 
 			CombinedLUTRenderTarget->GetDesc().Extent.Y != LUTSize ||
@@ -1177,22 +1204,7 @@ public:
 			!!(CombinedLUTRenderTarget->GetDesc().TargetableFlags & TexCreate_UAV) != bNeedUAV)
 		{
 			// Create the texture needed for the tonemapping LUT
-			EPixelFormat LUTPixelFormat = PF_A2B10G10R10;
-			if (!GPixelFormats[LUTPixelFormat].Supported)
-			{
-				LUTPixelFormat = PF_R8G8B8A8;
-			}
-
-			FPooledRenderTargetDesc Desc = FPooledRenderTargetDesc::Create2DDesc(FIntPoint(LUTSize * LUTSize, LUTSize), LUTPixelFormat, FClearValueBinding::Transparent, TexCreate_None, TexCreate_ShaderResource, false);
-			Desc.TargetableFlags |= bNeedUAV ? TexCreate_UAV : TexCreate_RenderTargetable;
-
-			if (bUseVolumeLUT)
-			{
-				Desc.Extent = FIntPoint(LUTSize, LUTSize);
-				Desc.Depth = LUTSize;
-			}
-
-			Desc.DebugName = TEXT("CombineLUTs");
+			FPooledRenderTargetDesc Desc = CreateLUTRenderTarget(LUTSize, bUseVolumeLUT, bNeedUAV, bNeedFloatOutput);
 
 			GRenderTargetPool.FindFreeElement(RHICmdList, Desc, CombinedLUTRenderTarget, Desc.DebugName);
 		}
diff --git a/Engine/Source/Runtime/Renderer/Private/SceneRendering.cpp b/Engine/Source/Runtime/Renderer/Private/SceneRendering.cpp
--- a/Engine/Source/Runtime/Renderer/Private/SceneRendering.cpp
+++ b/Engine/Source/Runtime/Renderer/Private/SceneRendering.cpp
@@ -1742,13 +1742,13 @@ const FTextureRHIRef* FViewInfo::GetTonemappingLUTTexture() const
 	return TextureRHIRef;
 };
 
-FSceneRenderTargetItem* FViewInfo::GetTonemappingLUTRenderTarget(FRHICommandList& RHICmdList, const int32 LUTSize, const bool bUseVolumeLUT, const bool bNeedUAV) const 
+FSceneRenderTargetItem* FViewInfo::GetTonemappingLUTRenderTarget(FRHICommandList& RHICmdList, const int32 LUTSize, const bool bUseVolumeLUT, const bool bNeedUAV, const bool bNeedFloatOutput) const 
 {
 	FSceneRenderTargetItem* TargetItem = NULL;
 	FSceneViewState* EffectiveViewState = GetEffectiveViewState();
 	if (EffectiveViewState)
 	{
-		TargetItem = &(EffectiveViewState->GetTonemappingLUTRenderTarget(RHICmdList, LUTSize, bUseVolumeLUT, bNeedUAV));
+		TargetItem = &(EffectiveViewState->GetTonemappingLUTRenderTarget(RHICmdList, LUTSize, bUseVolumeLUT, bNeedUAV, bNeedFloatOutput));
 	}
 	return TargetItem;
 }
diff --git a/Engine/Source/Runtime/Renderer/Private/SceneRendering.h b/Engine/Source/Runtime/Renderer/Private/SceneRendering.h
--- a/Engine/Source/Runtime/Renderer/Private/SceneRendering.h
+++ b/Engine/Source/Runtime/Renderer/Private/SceneRendering.h
@@ -1091,7 +1091,7 @@ public:
 
 	/** Gets the rendertarget that will be populated by CombineLUTS post process 
 	* for stereo rendering, this will force the post-processing to use the same render target for both eyes*/
-	FSceneRenderTargetItem* GetTonemappingLUTRenderTarget(FRHICommandList& RHICmdList, const int32 LUTSize, const bool bUseVolumeLUT, const bool bNeedUAV) const;
+	FSceneRenderTargetItem* GetTonemappingLUTRenderTarget(FRHICommandList& RHICmdList, const int32 LUTSize, const bool bUseVolumeLUT, const bool bNeedUAV, const bool bNeedFloatOutput) const;
 	
 
 
