Fix EditConditions in Composure elements

diff --git a/Engine/Plugins/Compositing/Composure/Source/Composure/Classes/CompositingElements/CompositingElementTransforms.h b/Engine/Plugins/Compositing/Composure/Source/Composure/Classes/CompositingElements/CompositingElementTransforms.h
--- a/Engine/Plugins/Compositing/Composure/Source/Composure/Classes/CompositingElements/CompositingElementTransforms.h
+++ b/Engine/Plugins/Compositing/Composure/Source/Composure/Classes/CompositingElements/CompositingElementTransforms.h
@@ -24,7 +24,7 @@ public:
 	UPROPERTY(EditAnywhere, BlueprintReadOnly, Category = "Compositing Pass", meta = (DisplayAfter = "PassName", EditCondition = "bEnabled"))
 	float RenderScale = 1.f;
 
-	UPROPERTY(EditAnywhere, Instanced, BlueprintReadOnly, Category = "Compositing Pass", meta = (DisplayAfter = "RenderScale"/*, EditCondition = "bEnabled"*/))
+	UPROPERTY(EditAnywhere, Instanced, BlueprintReadOnly, Category = "Compositing Pass", meta = (DisplayAfter = "RenderScale"))
 	TArray<UComposurePostProcessPassPolicy*> PostProcessPasses;
 
 public:
diff --git a/Engine/Plugins/Compositing/Composure/Source/Composure/Public/CompositingElements/CompositingMaterialPass.h b/Engine/Plugins/Compositing/Composure/Source/Composure/Public/CompositingElements/CompositingMaterialPass.h
--- a/Engine/Plugins/Compositing/Composure/Source/Composure/Public/CompositingElements/CompositingMaterialPass.h
+++ b/Engine/Plugins/Compositing/Composure/Source/Composure/Public/CompositingElements/CompositingMaterialPass.h
@@ -89,14 +89,14 @@ struct COMPOSURE_API FCompositingMaterial : public FCompositingParamPayload
 public:
 	FCompositingMaterial();
 
-	UPROPERTY(EditAnywhere, BlueprintReadWrite, Category = "CompositingMaterial", meta = (EditCondition = "bEnabled"))
+	UPROPERTY(EditAnywhere, BlueprintReadWrite, Category = "CompositingMaterial")
 	UMaterialInterface* Material;
 
 	/** Maps material texture param names to prior passes/elements. Overrides the element's param mapping list above. */
-	UPROPERTY(EditAnywhere, BlueprintReadOnly, Category = "CompositingMaterial", meta = (ReadOnlyKeys, EditCondition = "bEnabled"))
+	UPROPERTY(EditAnywhere, BlueprintReadOnly, Category = "CompositingMaterial", meta = (ReadOnlyKeys))
 	TMap<FName, FName> ParamPassMappings;
 
-	UPROPERTY(EditDefaultsOnly, Category = "CompositingMaterial", meta = (DisplayName="Expected Param Mappings", EditCondition = "bEnabled"))
+	UPROPERTY(EditDefaultsOnly, Category = "CompositingMaterial", meta = (DisplayName="Expected Param Mappings"))
 	TMap<FName, FNamedCompMaterialParam> RequiredMaterialParams;
 
 public:
diff --git a/Engine/Plugins/Compositing/Composure/Source/Composure/Public/MovieScene/MovieSceneComposureExportTrack.h b/Engine/Plugins/Compositing/Composure/Source/Composure/Public/MovieScene/MovieSceneComposureExportTrack.h
--- a/Engine/Plugins/Compositing/Composure/Source/Composure/Public/MovieScene/MovieSceneComposureExportTrack.h
+++ b/Engine/Plugins/Compositing/Composure/Source/Composure/Public/MovieScene/MovieSceneComposureExportTrack.h
@@ -24,7 +24,7 @@ struct FMovieSceneComposureExportPass
 	bool bRenamePass;
 
 	/** The name to use for this pass when rendering out */
-	UPROPERTY(EditAnywhere, Category="Export", meta=(EditCondition=bRenamePass))
+	UPROPERTY(EditAnywhere, Category="Export", meta=(EditCondition="bRenamePass"))
 	FName ExportedAs;
 };
 
