Fix for memory stomp in Compositing plugin opening editor with accessing TMap references after an Add

diff --git a/Engine/Plugins/Compositing/Composure/Source/Composure/Private/CompositingElements/CompositingElementTransforms.cpp b/Engine/Plugins/Compositing/Composure/Source/Composure/Private/CompositingElements/CompositingElementTransforms.cpp
--- a/Engine/Plugins/Compositing/Composure/Source/Composure/Private/CompositingElements/CompositingElementTransforms.cpp
+++ b/Engine/Plugins/Compositing/Composure/Source/Composure/Private/CompositingElements/CompositingElementTransforms.cpp
@@ -224,14 +224,14 @@ UMultiPassChromaKeyer::UMultiPassChromaKeyer()
 	static ConstructorHelpers::FObjectFinder<UMaterialInterface> GKeyerMaterial(TEXT("/Composure/Materials/ChromaKeying/M_SinglePassChromaKeyer"));
 	KeyerMaterial.Material = Cast<UMaterialInterface>(GKeyerMaterial.Object);
 
-	FNamedCompMaterialParam& InputParamMapping  = KeyerMaterial.RequiredMaterialParams.Add(MultiPassChromaKeyer_Impl::ColorPlateKeyName, TEXT("LinearColorPlate"));
-	FNamedCompMaterialParam& ColorParamMapping  = KeyerMaterial.RequiredMaterialParams.Add(MultiPassChromaKeyer_Impl::KeyColorKeyName, TEXT("KeyColor"));
-	FNamedCompMaterialParam& ResultParamMapping = KeyerMaterial.RequiredMaterialParams.Add(MultiPassChromaKeyer_Impl::KeyedResultInputName, TEXT("PrevKeyerResult"));
+	KeyerMaterial.RequiredMaterialParams.Add(MultiPassChromaKeyer_Impl::ColorPlateKeyName, TEXT("LinearColorPlate"));
+	KeyerMaterial.RequiredMaterialParams.Add(MultiPassChromaKeyer_Impl::KeyColorKeyName, TEXT("KeyColor"));
+	KeyerMaterial.RequiredMaterialParams.Add(MultiPassChromaKeyer_Impl::KeyedResultInputName, TEXT("PrevKeyerResult"));
 
 #if WITH_EDITOR
-	InputParamMapping.ParamType = EParamType::TextureParam;
-	ColorParamMapping.ParamType = EParamType::VectorParam;
-	ResultParamMapping.ParamType = EParamType::TextureParam;
+	KeyerMaterial.RequiredMaterialParams[MultiPassChromaKeyer_Impl::ColorPlateKeyName].ParamType = EParamType::TextureParam;
+	KeyerMaterial.RequiredMaterialParams[MultiPassChromaKeyer_Impl::KeyColorKeyName].ParamType = EParamType::VectorParam;
+	KeyerMaterial.RequiredMaterialParams[MultiPassChromaKeyer_Impl::KeyedResultInputName].ParamType = EParamType::TextureParam;
 #endif
 
 	static ConstructorHelpers::FObjectFinder<UTexture> GDefaultResultTexture(TEXT("/Engine/EngineResources/WhiteSquareTexture.WhiteSquareTexture"));
@@ -290,14 +290,14 @@ UMultiPassDespill::UMultiPassDespill()
 	static ConstructorHelpers::FObjectFinder<UMaterialInterface> GKeyerMaterial(TEXT("/Composure/Materials/ChromaKeying/M_SinglePassDespill"));
 	KeyerMaterial.Material = Cast<UMaterialInterface>(GKeyerMaterial.Object);
 
-	FNamedCompMaterialParam& InputParamMapping  = KeyerMaterial.RequiredMaterialParams.Add(MultiPassDespill_Impl::ColorPlateKeyName, TEXT("LinearColorPlate"));
-	FNamedCompMaterialParam& ColorParamMapping  = KeyerMaterial.RequiredMaterialParams.Add(MultiPassDespill_Impl::KeyColorKeyName, TEXT("KeyColor"));
-	FNamedCompMaterialParam& ResultParamMapping = KeyerMaterial.RequiredMaterialParams.Add(MultiPassDespill_Impl::KeyedResultInputName, TEXT("PrevKeyerResult"));
+	KeyerMaterial.RequiredMaterialParams.Add(MultiPassDespill_Impl::ColorPlateKeyName, TEXT("LinearColorPlate"));
+	KeyerMaterial.RequiredMaterialParams.Add(MultiPassDespill_Impl::KeyColorKeyName, TEXT("KeyColor"));
+	KeyerMaterial.RequiredMaterialParams.Add(MultiPassDespill_Impl::KeyedResultInputName, TEXT("PrevKeyerResult"));
 
 #if WITH_EDITOR
-	InputParamMapping.ParamType = EParamType::TextureParam;
-	ColorParamMapping.ParamType = EParamType::VectorParam;
-	ResultParamMapping.ParamType = EParamType::TextureParam;
+	KeyerMaterial.RequiredMaterialParams[MultiPassDespill_Impl::ColorPlateKeyName].ParamType = EParamType::TextureParam;
+	KeyerMaterial.RequiredMaterialParams[MultiPassDespill_Impl::KeyColorKeyName].ParamType = EParamType::VectorParam;
+	KeyerMaterial.RequiredMaterialParams[MultiPassDespill_Impl::KeyedResultInputName].ParamType = EParamType::TextureParam;
 #endif
 
 	static ConstructorHelpers::FObjectFinder<UTexture> GDefaultResultTexture(TEXT("/Engine/Functions/Engine_MaterialFunctions02/PivotPainter2/Black_1x1_EXR_Texture.Black_1x1_EXR_Texture"));
