Workaround for static analysis issue warning about uninitialized memory. For some reason the SA complains about copy construction even though the default constructor *looks* ok

diff --git a/Engine/Plugins/Compositing/Composure/Source/ComposureLayersEditor/Private/Widgets/SCompElementPickerWindow.cpp b/Engine/Plugins/Compositing/Composure/Source/ComposureLayersEditor/Private/Widgets/SCompElementPickerWindow.cpp
--- a/Engine/Plugins/Compositing/Composure/Source/ComposureLayersEditor/Private/Widgets/SCompElementPickerWindow.cpp
+++ b/Engine/Plugins/Compositing/Composure/Source/ComposureLayersEditor/Private/Widgets/SCompElementPickerWindow.cpp
@@ -582,7 +582,9 @@ void SCompElementPickerWindow::Construct(const FArguments& InArgs)
 			.SupportsMinimize(false)
 	);		
 
-	TWeakUIntrfacePtr<ICompImageColorPickerInterface> PickerTarget = InArgs._PickerTarget;
+	TWeakUIntrfacePtr<ICompImageColorPickerInterface> PickerTarget; 
+	PickerTarget = InArgs._PickerTarget;
+
 	FColorPickedEventHandler OnColorPicked = InArgs._OnColorPicked;
 	OnPickerCanceled = InArgs._OnPickerCanceled;
 
