Composure: Editor updates & fixes for UE5.

Correctly showing icons in the composure layer editor by overriding new padding defaults. Fixed preview window color and padding. Removing UE4 mention in plugin description.

diff --git a/Engine/Plugins/Compositing/Composure/Composure.uplugin b/Engine/Plugins/Compositing/Composure/Composure.uplugin
--- a/Engine/Plugins/Compositing/Composure/Composure.uplugin
+++ b/Engine/Plugins/Compositing/Composure/Composure.uplugin
@@ -3,7 +3,7 @@
 	"Version" : 1,
 	"VersionName" : "1.0",
 	"FriendlyName" : "Composure",
-	"Description" : "Plugin design to make compositing easier in UE4.",
+	"Description" : "Plugin design to make compositing easier in Unreal Engine.",
 	"Category" : "Compositing",
 	"CreatedBy" : "Epic Games, Inc.",
 	"CreatedByURL" : "http://epicgames.com",
diff --git a/Engine/Plugins/Compositing/Composure/Source/ComposureLayersEditor/Private/Widgets/SCompElementPreviewDialog.cpp b/Engine/Plugins/Compositing/Composure/Source/ComposureLayersEditor/Private/Widgets/SCompElementPreviewDialog.cpp
--- a/Engine/Plugins/Compositing/Composure/Source/ComposureLayersEditor/Private/Widgets/SCompElementPreviewDialog.cpp
+++ b/Engine/Plugins/Compositing/Composure/Source/ComposureLayersEditor/Private/Widgets/SCompElementPreviewDialog.cpp
@@ -86,8 +86,7 @@ void SCompElementPreviewDialog::Construct(const FArguments& InArgs)
 	ChildSlot
 	[
 		SNew(SBorder)
-			.Padding(FMargin(0.f))
-			.BorderImage(FEditorStyle::GetBrush("DetailsView.CategoryMiddle"))
+			.Padding(FMargin(2.f))
 			.Content()
 		[
 			GeneratePreviewContent()	
diff --git a/Engine/Plugins/Compositing/Composure/Source/ComposureLayersEditor/Private/Widgets/SCompElementViewRow.cpp b/Engine/Plugins/Compositing/Composure/Source/ComposureLayersEditor/Private/Widgets/SCompElementViewRow.cpp
--- a/Engine/Plugins/Compositing/Composure/Source/ComposureLayersEditor/Private/Widgets/SCompElementViewRow.cpp
+++ b/Engine/Plugins/Compositing/Composure/Source/ComposureLayersEditor/Private/Widgets/SCompElementViewRow.cpp
@@ -163,8 +163,8 @@ TSharedRef<SWidget> SCompElementViewRow::GenerateWidgetForColumn(const FName& Co
 		};
 
 		TableRowContent = SNew(SBox)
-			.MinDesiredWidth(66.f)
-			.MaxDesiredWidth(66.f)
+			.MinDesiredWidth(114.f)
+			.MaxDesiredWidth(114.f)
 			[
 				SNew(SComboButton)
 					.ContentPadding(FMargin(0, 0, 5, 0))
@@ -206,7 +206,7 @@ TSharedRef<SWidget> SCompElementViewRow::GenerateWidgetForColumn(const FName& Co
 		TableRowContent =
 			SAssignNew(MediaCaptureToggle, SContextMenuButton)
 			.ContentPadding(0)
-			.ButtonStyle(FEditorStyle::Get(), "ToggleButton")
+			.ButtonStyle(FEditorStyle::Get(), "NoBorder")
 			.ToolTipText(LOCTEXT("MediaCaptureToggleTooltip", "Turn Media Capture On/Off"))
 			.OnClicked(this, &SCompElementViewRow::OnToggleMediaCapture)
 			.IsEnabled(this, &SCompElementViewRow::IsMediaCaptureToggleEnabled)
@@ -230,7 +230,7 @@ TSharedRef<SWidget> SCompElementViewRow::GenerateWidgetForColumn(const FName& Co
 		TableRowContent =
 			SAssignNew(FreezeFrameButton, SButton)
 			.ContentPadding(0)
-			.ButtonStyle(FEditorStyle::Get(), "ToggleButton")
+			.ButtonStyle(FEditorStyle::Get(), "NoBorder")
 			.ToolTipText(LOCTEXT("FreezeToggleTooltip", "Toggle Freeze Framing"))
 			.OnClicked(this, &SCompElementViewRow::OnToggleFreezeFrame)
 			.IsEnabled(this, &SCompElementViewRow::IsFreezeFrameToggleEnabled)
diff --git a/Engine/Plugins/Compositing/Composure/Source/ComposureLayersEditor/Private/Widgets/SCompElementsView.h b/Engine/Plugins/Compositing/Composure/Source/ComposureLayersEditor/Private/Widgets/SCompElementsView.h
--- a/Engine/Plugins/Compositing/Composure/Source/ComposureLayersEditor/Private/Widgets/SCompElementsView.h
+++ b/Engine/Plugins/Compositing/Composure/Source/ComposureLayersEditor/Private/Widgets/SCompElementsView.h
@@ -85,7 +85,7 @@ public:
 			+ SHeaderRow::Column(CompElementsView::ColumnID_Alpha)
 				.HAlignCell(HAlign_Right)
 				.HAlignHeader(HAlign_Right)
-				.FixedWidth(66.f)
+				.FixedWidth(114.f)
 				.DefaultLabel(LOCTEXT("Column_AlphaNameLabel", "Alpha"))
 
 			+ SHeaderRow::Column(CompElementsView::ColumnID_MediaCapture)
