Composure: Fix for missing input/output passes by preventing unwanted class specifier inheritance (Abstract suffices).

diff --git a/Engine/Plugins/Compositing/Composure/Source/Composure/Classes/CompositingElements/CompositingElementInputs.h b/Engine/Plugins/Compositing/Composure/Source/Composure/Classes/CompositingElements/CompositingElementInputs.h
--- a/Engine/Plugins/Compositing/Composure/Source/Composure/Classes/CompositingElements/CompositingElementInputs.h
+++ b/Engine/Plugins/Compositing/Composure/Source/Composure/Classes/CompositingElements/CompositingElementInputs.h
@@ -12,7 +12,7 @@
 
 class UMaterialInterface;
 
-UCLASS(noteditinlinenew, hidedropdown, Abstract)
+UCLASS(Abstract)
 class COMPOSURE_API UCompositingMediaInput : public UCompositingElementInput
 {
 	GENERATED_BODY()
diff --git a/Engine/Plugins/Compositing/Composure/Source/Composure/Classes/CompositingElements/CompositingElementOutputs.h b/Engine/Plugins/Compositing/Composure/Source/Composure/Classes/CompositingElements/CompositingElementOutputs.h
--- a/Engine/Plugins/Compositing/Composure/Source/Composure/Classes/CompositingElements/CompositingElementOutputs.h
+++ b/Engine/Plugins/Compositing/Composure/Source/Composure/Classes/CompositingElements/CompositingElementOutputs.h
@@ -16,7 +16,7 @@ class UMediaCapture;
 class UCompositingElementTransform;
 class UTextureRenderTarget2D;
 
-UCLASS(noteditinlinenew, hidedropdown, Abstract)
+UCLASS(Abstract)
 class COMPOSURE_API UColorConverterOutputPass : public UCompositingElementOutput, public ICompEditorImagePreviewInterface
 {
 	GENERATED_BODY()
