Sequencer: Ensure for null terminal player in RestoreState

diff --git a/Engine/Plugins/Compositing/Composure/Source/Composure/Private/MovieScene/MovieSceneComposureExportSectionTemplate.cpp b/Engine/Plugins/Compositing/Composure/Source/Composure/Private/MovieScene/MovieSceneComposureExportSectionTemplate.cpp
--- a/Engine/Plugins/Compositing/Composure/Source/Composure/Private/MovieScene/MovieSceneComposureExportSectionTemplate.cpp
+++ b/Engine/Plugins/Compositing/Composure/Source/Composure/Private/MovieScene/MovieSceneComposureExportSectionTemplate.cpp
@@ -164,7 +164,11 @@ struct FStopCapturingShotElementPassToken : IMovieScenePreAnimatedToken
 	virtual void RestoreState(UObject& Object, const UE::MovieScene::FRestoreStateParams& Params)
 	{
 		IMovieScenePlayer* Player = Params.GetTerminalPlayer();
-		check(Player);
+		if (!ensure(Player))
+		{
+			return;
+		}
+
 		FPersistentEvaluationData PersistentData(*Player);
 
 		FComposureShotElementCaptureManager* CaptureManager = PersistentData.Find<FComposureShotElementCaptureManager>(ComposureExportSharedKey);
diff --git a/Engine/Source/Runtime/MovieScene/Public/Evaluation/MovieScenePropertyTemplate.h b/Engine/Source/Runtime/MovieScene/Public/Evaluation/MovieScenePropertyTemplate.h
--- a/Engine/Source/Runtime/MovieScene/Public/Evaluation/MovieScenePropertyTemplate.h
+++ b/Engine/Source/Runtime/MovieScene/Public/Evaluation/MovieScenePropertyTemplate.h
@@ -86,7 +86,13 @@ namespace PropertyTemplate
 
 		virtual void RestoreState(UObject& Object, const UE::MovieScene::FRestoreStateParams& Params) override
 		{
-			auto NewValue = PropertyTemplate::ConvertFromIntermediateType<PropertyValueType, IntermediateType>(Value, *Params.GetTerminalPlayer());
+			IMovieScenePlayer* Player = Params.GetTerminalPlayer();
+			if (!ensure(Player))
+			{
+				return;
+			}
+
+			auto NewValue = PropertyTemplate::ConvertFromIntermediateType<PropertyValueType, IntermediateType>(Value, *Player);
 			if (IsValueValid(NewValue))
 			{
 				Bindings.CallFunction<PropertyValueType>(Object, NewValue);
diff --git a/Engine/Source/Runtime/MovieSceneTracks/Private/Evaluation/MovieSceneCameraShakeSourceShakeTemplate.cpp b/Engine/Source/Runtime/MovieSceneTracks/Private/Evaluation/MovieSceneCameraShakeSourceShakeTemplate.cpp
--- a/Engine/Source/Runtime/MovieSceneTracks/Private/Evaluation/MovieSceneCameraShakeSourceShakeTemplate.cpp
+++ b/Engine/Source/Runtime/MovieSceneTracks/Private/Evaluation/MovieSceneCameraShakeSourceShakeTemplate.cpp
@@ -43,7 +43,13 @@ struct FPreAnimatedCameraShakeSourceShakeTokenProducer : IMovieScenePreAnimatedG
 
 			virtual void RestoreState(const UE::MovieScene::FRestoreStateParams& Params) override
 			{
-				for (TWeakObjectPtr<> BoundObject : Params.GetTerminalPlayer()->FindBoundObjects(Operand))
+				IMovieScenePlayer* Player = Params.GetTerminalPlayer();
+				if (!ensure(Player))
+				{
+					return;
+				}
+
+				for (TWeakObjectPtr<> BoundObject : Player->FindBoundObjects(Operand))
 				{
 					if (BoundObject.IsValid())
 					{
diff --git a/Engine/Source/Runtime/MovieSceneTracks/Private/Evaluation/MovieSceneFadeTemplate.cpp b/Engine/Source/Runtime/MovieSceneTracks/Private/Evaluation/MovieSceneFadeTemplate.cpp
--- a/Engine/Source/Runtime/MovieSceneTracks/Private/Evaluation/MovieSceneFadeTemplate.cpp
+++ b/Engine/Source/Runtime/MovieSceneTracks/Private/Evaluation/MovieSceneFadeTemplate.cpp
@@ -62,7 +62,13 @@ struct FFadePreAnimatedGlobalToken : FFadeTrackToken, IMovieScenePreAnimatedGlob
 
 	virtual void RestoreState(const UE::MovieScene::FRestoreStateParams& Params) override
 	{
-		Apply(*Params.GetTerminalPlayer(), true);
+		IMovieScenePlayer* Player = Params.GetTerminalPlayer();
+		if (!ensure(Player))
+		{
+			return;
+		}
+		
+		Apply(*Player, true);
 	}
 };
 
diff --git a/Engine/Source/Runtime/MovieSceneTracks/Private/Evaluation/MovieSceneSlomoTemplate.cpp b/Engine/Source/Runtime/MovieSceneTracks/Private/Evaluation/MovieSceneSlomoTemplate.cpp
--- a/Engine/Source/Runtime/MovieSceneTracks/Private/Evaluation/MovieSceneSlomoTemplate.cpp
+++ b/Engine/Source/Runtime/MovieSceneTracks/Private/Evaluation/MovieSceneSlomoTemplate.cpp
@@ -55,7 +55,13 @@ struct FSlomoPreAnimatedGlobalToken : FSlomoTrackToken, IMovieScenePreAnimatedGl
 
 	virtual void RestoreState(const UE::MovieScene::FRestoreStateParams& Params) override
 	{
-		Apply(*Params.GetTerminalPlayer());
+		IMovieScenePlayer* Player = Params.GetTerminalPlayer();
+		if (!ensure(Player))
+		{
+			return;
+		}
+
+		Apply(*Player);
 	}
 };
 
diff --git a/Engine/Source/Runtime/MovieSceneTracks/Private/TrackInstances/MovieSceneCameraCutTrackInstance.cpp b/Engine/Source/Runtime/MovieSceneTracks/Private/TrackInstances/MovieSceneCameraCutTrackInstance.cpp
--- a/Engine/Source/Runtime/MovieSceneTracks/Private/TrackInstances/MovieSceneCameraCutTrackInstance.cpp
+++ b/Engine/Source/Runtime/MovieSceneTracks/Private/TrackInstances/MovieSceneCameraCutTrackInstance.cpp
@@ -94,8 +94,14 @@ namespace MovieScene
 
 		virtual void RestoreState(const UE::MovieScene::FRestoreStateParams& RestoreParams) override
 		{
+			IMovieScenePlayer* Player = RestoreParams.GetTerminalPlayer();
+			if (!ensure(Player))
+			{
+				return;
+			}
+			
 			EMovieSceneCameraCutParams Params;
-			RestoreParams.GetTerminalPlayer()->UpdateCameraCut(nullptr, Params);
+			Player->UpdateCameraCut(nullptr, Params);
 		}
 	};
 
