Composure: Added a way to create transient elements.

diff --git a/Engine/Plugins/Compositing/Composure/Source/ComposureLayersEditor/Private/CompElementManager.cpp b/Engine/Plugins/Compositing/Composure/Source/ComposureLayersEditor/Private/CompElementManager.cpp
--- a/Engine/Plugins/Compositing/Composure/Source/ComposureLayersEditor/Private/CompElementManager.cpp
+++ b/Engine/Plugins/Compositing/Composure/Source/ComposureLayersEditor/Private/CompElementManager.cpp
@@ -78,7 +78,7 @@ void FCompElementManager::Initialize()
 	}
 }
 
-TWeakObjectPtr<ACompositingElement> FCompElementManager::CreateElement(const FName& ElementName, TSubclassOf<ACompositingElement> ClassType, AActor* LevelContext)
+TWeakObjectPtr<ACompositingElement> FCompElementManager::CreateElement(const FName& ElementName, TSubclassOf<ACompositingElement> ClassType, AActor* LevelContext, UPackage* OverridePackage)
 {
 	ACompositingElement* SpawnedActor = nullptr;
 
@@ -109,6 +109,9 @@ TWeakObjectPtr<ACompositingElement> FCompElementManager::CreateElement(const FNa
 		FScopedWorldLevelContext ScopedLevelContext(TargetWorld, LevelContext);
 
 		FActorSpawnParameters SpawnParams;
+#if WITH_EDITOR
+		SpawnParams.OverridePackage = OverridePackage;
+#endif
 		SpawnedActor = TargetWorld->SpawnActor<ACompositingElement>(ClassType, FTransform::Identity, SpawnParams);
 
 		if (SpawnedActor)
diff --git a/Engine/Plugins/Compositing/Composure/Source/ComposureLayersEditor/Private/CompElementManager.h b/Engine/Plugins/Compositing/Composure/Source/ComposureLayersEditor/Private/CompElementManager.h
--- a/Engine/Plugins/Compositing/Composure/Source/ComposureLayersEditor/Private/CompElementManager.h
+++ b/Engine/Plugins/Compositing/Composure/Source/ComposureLayersEditor/Private/CompElementManager.h
@@ -9,6 +9,7 @@
 
 class UEditorCompElementContainer;
 class AActor;
+class UPackage;
 class UEditorEngine;
 class FCompositingViewportClient;
 class ULevel;
@@ -33,7 +34,7 @@ public:
 
 public: 
 	//~ Begin ICompElementManager interface
-	virtual TWeakObjectPtr<ACompositingElement> CreateElement(const FName& ElementName, TSubclassOf<ACompositingElement> ClassType, AActor* LevelContext = nullptr) override;
+	virtual TWeakObjectPtr<ACompositingElement> CreateElement(const FName& ElementName, TSubclassOf<ACompositingElement> ClassType, AActor* LevelContext = nullptr, UPackage* OverridePackage = nullptr) override;
 	virtual TWeakObjectPtr<ACompositingElement> GetElement(const FName& ElementName) const override;
 	virtual bool TryGetElement(const FName& ElementName, TWeakObjectPtr<ACompositingElement>& OutElement) override;
 	virtual void AddAllCompElementsTo(TArray< TWeakObjectPtr<ACompositingElement> >& OutElements) const override;
diff --git a/Engine/Plugins/Compositing/Composure/Source/ComposureLayersEditor/Public/ICompElementManager.h b/Engine/Plugins/Compositing/Composure/Source/ComposureLayersEditor/Public/ICompElementManager.h
--- a/Engine/Plugins/Compositing/Composure/Source/ComposureLayersEditor/Public/ICompElementManager.h
+++ b/Engine/Plugins/Compositing/Composure/Source/ComposureLayersEditor/Public/ICompElementManager.h
@@ -7,6 +7,7 @@
 class ACompositingElement;
 class UCompositingMediaCaptureOutput;
 class AActor;
+class UPackage;
 template<typename TItemType> class IFilter;
 
 /** */
@@ -35,7 +36,7 @@ public:
 	 * @param  LevelContext Optional actor whose world level you want to spawn into.
 	 * @return The newly created actor object for the named element.
 	 */
-	virtual TWeakObjectPtr<ACompositingElement> CreateElement(const FName& ElementName, TSubclassOf<ACompositingElement> ClassType, AActor* LevelContext = nullptr) = 0;
+	virtual TWeakObjectPtr<ACompositingElement> CreateElement(const FName& ElementName, TSubclassOf<ACompositingElement> ClassType, AActor* LevelContext = nullptr, UPackage* OverridePackage = nullptr) = 0;
  
 	/**
 	 * Gets the actor object of the named element.
