Composure: Added distortion setters to CompositingCaptureBase.

This avoids going through property reflection.

diff --git a/Engine/Plugins/Compositing/Composure/Source/Composure/Classes/CompositingCaptureBase.h b/Engine/Plugins/Compositing/Composure/Source/Composure/Classes/CompositingCaptureBase.h
--- a/Engine/Plugins/Compositing/Composure/Source/Composure/Classes/CompositingCaptureBase.h
+++ b/Engine/Plugins/Compositing/Composure/Source/Composure/Classes/CompositingCaptureBase.h
@@ -5,11 +5,11 @@
 #include "CompositingElement.h"
 
 #include "CameraCalibrationTypes.h"
-#include "LensDistortionModelHandlerBase.h"
-
 #include "CompositingCaptureBase.generated.h"
 
 class USceneCaptureComponent2D;
+class ULensDistortionModelHandlerBase;
+
 
 /**
  * Base class for CG Compositing Elements
@@ -42,6 +42,7 @@ protected:
 	UMaterialInstanceDynamic* LastDistortionMID = nullptr;
 
 public:
+
 	/** Default constructor */
 	ACompositingCaptureBase();
 
@@ -54,4 +55,12 @@ public:
 	virtual void PostEditChangeProperty(struct FPropertyChangedEvent& PropertyChangedEvent) override;
 #endif	
 	//~ End UObject Interface
-};
\ No newline at end of file
+
+public:
+
+	/** Sets whether distortion should be applied or not */
+	void SetApplyDistortion(bool bInApplyDistortion);
+
+	/** Sets which distortion handler to use when bInApplyDistortion is enabled */
+	void SetDistortionHandler(ULensDistortionModelHandlerBase* InDistortionHandler);
+};
diff --git a/Engine/Plugins/Compositing/Composure/Source/Composure/Private/CompositingCaptureBase.cpp b/Engine/Plugins/Compositing/Composure/Source/Composure/Private/CompositingCaptureBase.cpp
--- a/Engine/Plugins/Compositing/Composure/Source/Composure/Private/CompositingCaptureBase.cpp
+++ b/Engine/Plugins/Compositing/Composure/Source/Composure/Private/CompositingCaptureBase.cpp
@@ -6,6 +6,8 @@
 #include "CineCameraComponent.h"
 #include "Components/SceneCaptureComponent2D.h"
 #include "Engine/Engine.h"
+#include "LensDistortionModelHandlerBase.h"
+
 
 ACompositingCaptureBase::ACompositingCaptureBase()
 {
@@ -126,3 +128,25 @@ void ACompositingCaptureBase::PostEditChangeProperty(struct FPropertyChangedEven
 	Super::PostEditChangeProperty(PropertyChangedEvent);
 }
 #endif
+
+void ACompositingCaptureBase::SetApplyDistortion(bool bInApplyDistortion)
+{
+	bApplyDistortion = bInApplyDistortion;
+	UpdateDistortion();
+}
+
+void ACompositingCaptureBase::SetDistortionHandler(ULensDistortionModelHandlerBase* InDistortionHandler)
+{
+	if (InDistortionHandler)
+	{
+		DistortionSource.DistortionProducerID = InDistortionHandler->GetDistortionProducerID();
+		DistortionSource.HandlerDisplayName = InDistortionHandler->GetDisplayName();
+	}
+	else
+	{
+		DistortionSource.DistortionProducerID = FGuid();
+		DistortionSource.HandlerDisplayName = FString();
+	}
+
+	UpdateDistortion();
+}
\ No newline at end of file
