Composure: Added a distortion handler getter.

diff --git a/Engine/Plugins/Compositing/Composure/Source/Composure/Classes/CompositingCaptureBase.h b/Engine/Plugins/Compositing/Composure/Source/Composure/Classes/CompositingCaptureBase.h
--- a/Engine/Plugins/Compositing/Composure/Source/Composure/Classes/CompositingCaptureBase.h
+++ b/Engine/Plugins/Compositing/Composure/Source/Composure/Classes/CompositingCaptureBase.h
@@ -63,4 +63,7 @@ public:
 
 	/** Sets which distortion handler to use when bInApplyDistortion is enabled */
 	void SetDistortionHandler(ULensDistortionModelHandlerBase* InDistortionHandler);
+
+	/** Returns the current distortion handler */
+	ULensDistortionModelHandlerBase* GetDistortionHandler();
 };
diff --git a/Engine/Plugins/Compositing/Composure/Source/Composure/Private/CompositingCaptureBase.cpp b/Engine/Plugins/Compositing/Composure/Source/Composure/Private/CompositingCaptureBase.cpp
--- a/Engine/Plugins/Compositing/Composure/Source/Composure/Private/CompositingCaptureBase.cpp
+++ b/Engine/Plugins/Compositing/Composure/Source/Composure/Private/CompositingCaptureBase.cpp
@@ -149,4 +149,16 @@ void ACompositingCaptureBase::SetDistortionHandler(ULensDistortionModelHandlerBa
 	}
 
 	UpdateDistortion();
-}
\ No newline at end of file
+}
+
+ULensDistortionModelHandlerBase* ACompositingCaptureBase::GetDistortionHandler()
+{
+	UCameraCalibrationSubsystem* SubSystem = GEngine->GetEngineSubsystem<UCameraCalibrationSubsystem>();
+
+	if (!SubSystem)
+	{
+		return nullptr;
+	}
+
+	return SubSystem->FindDistortionModelHandler(DistortionSource);
+}
