Composure: Fix RelayOutput_Implementation for the player viewport output, where the color transform result was never connected properly.

diff --git a/Engine/Plugins/Compositing/Composure/Source/Composure/Private/CompositingElements/PlayerViewportCompositingOutput.cpp b/Engine/Plugins/Compositing/Composure/Source/Composure/Private/CompositingElements/PlayerViewportCompositingOutput.cpp
--- a/Engine/Plugins/Compositing/Composure/Source/Composure/Private/CompositingElements/PlayerViewportCompositingOutput.cpp
+++ b/Engine/Plugins/Compositing/Composure/Source/Composure/Private/CompositingElements/PlayerViewportCompositingOutput.cpp
@@ -158,7 +158,7 @@ void UPlayerViewportCompositingOutput::RelayOutput_Implementation(UTexture* Rend
 			}
 		}
 
-		ViewportOverrideMID->SetTextureParameterValue(PlayerViewportCompDispatcher_Impl::PlayerViewportMatInputName, RenderResult);
+		ViewportOverrideMID->SetTextureParameterValue(PlayerViewportCompDispatcher_Impl::PlayerViewportMatInputName, OutputImage);
 	}
 	else
 	{
