Composure: Ensure that all compositing element transforms export their symbols.

#rnx

diff --git a/Engine/Plugins/Compositing/Composure/Source/Composure/Classes/CompositingElements/CompositingElementTransforms.h b/Engine/Plugins/Compositing/Composure/Source/Composure/Classes/CompositingElements/CompositingElementTransforms.h
--- a/Engine/Plugins/Compositing/Composure/Source/Composure/Classes/CompositingElements/CompositingElementTransforms.h
+++ b/Engine/Plugins/Compositing/Composure/Source/Composure/Classes/CompositingElements/CompositingElementTransforms.h
@@ -114,7 +114,7 @@ class UMaterialInstanceDynamic;
 class UMaterialInterface;
 
 UCLASS(BlueprintType, Blueprintable)
-class UMultiPassChromaKeyer : public UCompositingElementTransform
+class COMPOSURE_API UMultiPassChromaKeyer : public UCompositingElementTransform
 {
 	GENERATED_BODY()
 
@@ -145,7 +145,7 @@ class UMaterialInstanceDynamic;
 class UMaterialInterface;
 
 UCLASS(BlueprintType, Blueprintable)
-class UMultiPassDespill : public UCompositingElementTransform
+class COMPOSURE_API UMultiPassDespill : public UCompositingElementTransform
 {
 	GENERATED_BODY()
 
@@ -173,7 +173,7 @@ private:
 class UMaterialInstanceDynamic;
 
 UCLASS(noteditinlinenew, hidedropdown)
-class UAlphaTransformPass : public UCompositingElementTransform
+class COMPOSURE_API UAlphaTransformPass : public UCompositingElementTransform
 {
 	GENERATED_BODY()
 
