Composure: Tone curve in viewport output now disabled by default, with option to re-enable.

diff --git a/Engine/Plugins/Compositing/Composure/Source/Composure/Classes/CompositingElements/PlayerViewportCompositingOutput.h b/Engine/Plugins/Compositing/Composure/Source/Composure/Classes/CompositingElements/PlayerViewportCompositingOutput.h
--- a/Engine/Plugins/Compositing/Composure/Source/Composure/Classes/CompositingElements/PlayerViewportCompositingOutput.h
+++ b/Engine/Plugins/Compositing/Composure/Source/Composure/Classes/CompositingElements/PlayerViewportCompositingOutput.h
@@ -29,6 +29,10 @@ public:
 	UPROPERTY(EditAnywhere, Category="Compositing Pass", meta = (DisplayAfter = "PassName", EditCondition = "bEnabled"))
 	int32 PlayerIndex = 0;
 
+	/** Enable or disable the tone curve (and expand gamut) when the tonemap pass settings are passed down to the main viewport post-processing. */
+	UPROPERTY(EditAnywhere, Category = "Compositing Pass", meta = (DisplayAfter = "ColorConversion", EditCondition = "bEnabled"))
+	bool ApplyToneCurve = false;
+
 public:
 	//~ Begin UCompositingElementOutput interface
 	virtual void OnFrameBegin_Implementation(bool bCameraCutThisFrame) override;
@@ -40,6 +44,10 @@ public:
 	virtual void OverrideBlendableSettings(FSceneView& View, float Weight) const override;
 	//~ End IBlendableInterface interface
 
+
+#if WITH_EDITOR
+	virtual bool CanEditChange(const FProperty* InProperty) const override;
+#endif
 protected:
 	bool OverridePlayerCamera(int32 PlayerIndex);
 	void ClearViewportOverride();
diff --git a/Engine/Plugins/Compositing/Composure/Source/Composure/Private/CompositingElements/PlayerViewportCompositingOutput.cpp b/Engine/Plugins/Compositing/Composure/Source/Composure/Private/CompositingElements/PlayerViewportCompositingOutput.cpp
--- a/Engine/Plugins/Compositing/Composure/Source/Composure/Private/CompositingElements/PlayerViewportCompositingOutput.cpp
+++ b/Engine/Plugins/Compositing/Composure/Source/Composure/Private/CompositingElements/PlayerViewportCompositingOutput.cpp
@@ -187,6 +187,24 @@ void UPlayerViewportCompositingOutput::OverrideBlendableSettings(class FSceneVie
 	}
 }
 
+#if WITH_EDITOR
+bool UPlayerViewportCompositingOutput::CanEditChange(const FProperty* InProperty) const
+{
+	if (!Super::CanEditChange(InProperty))
+	{
+		return false;
+	}
+
+	if (InProperty->GetFName() == GET_MEMBER_NAME_CHECKED(UPlayerViewportCompositingOutput, ApplyToneCurve))
+	{
+		//Note: When ColorConverter is null, a tonemap pass is automatically created (based on the DefaultConverterClass set above)
+		return ColorConverter == nullptr || ColorConverter->GetClass() == UCompositingTonemapPass::StaticClass();
+	}
+
+	return true;
+}
+#endif
+
 bool UPlayerViewportCompositingOutput::OverridePlayerCamera(int32 InPlayerIndex)
 {
 	ClearViewportOverride();
@@ -274,6 +292,16 @@ bool UPlayerCompOutputCameraModifier::ModifyCamera(float /*DeltaTime*/, FMinimal
 				InOutPOV.PostProcessSettings
 			);
 		}
+
+		if (!Owner->ApplyToneCurve)
+		{
+			InOutPOV.PostProcessSettings.bOverride_ToneCurveAmount = true;
+			InOutPOV.PostProcessSettings.ToneCurveAmount = 0;
+
+			InOutPOV.PostProcessSettings.bOverride_ExpandGamut = true;
+			InOutPOV.PostProcessSettings.ExpandGamut = 0;
+		}
+
 		InOutPOV.PostProcessSettings.AddBlendable(Owner, /*Weight =*/1.0f);
 
 		return true;
