PR #9173: Fixed ACompositingElement is not calling Super::BeginPlay (Contributed by [REDACTED_PERSON])

diff --git a/Engine/Plugins/Compositing/Composure/Source/Composure/Private/CompositingElement.cpp b/Engine/Plugins/Compositing/Composure/Source/Composure/Private/CompositingElement.cpp
--- a/Engine/Plugins/Compositing/Composure/Source/Composure/Private/CompositingElement.cpp
+++ b/Engine/Plugins/Compositing/Composure/Source/Composure/Private/CompositingElement.cpp
@@ -908,6 +908,7 @@ void ACompositingElement::PostLoad()
 
 void ACompositingElement::BeginPlay()
 {
+	Super::BeginPlay();
 	RefreshAllInternalPassLists();
 }
 
