Added LLM tags

diff --git a/Engine/Plugins/Compositing/Composure/Source/ComposureEditor/Private/ComposureEditorModule.cpp b/Engine/Plugins/Compositing/Composure/Source/ComposureEditor/Private/ComposureEditorModule.cpp
--- a/Engine/Plugins/Compositing/Composure/Source/ComposureEditor/Private/ComposureEditorModule.cpp
+++ b/Engine/Plugins/Compositing/Composure/Source/ComposureEditor/Private/ComposureEditorModule.cpp
@@ -6,12 +6,14 @@
 #include "Sequencer/ComposurePostMoveSettingsPropertyTrackEditor.h"
 #include "Sequencer/ComposureExportTrackEditor.h"
 
+LLM_DEFINE_TAG(Compositing_ComposureEditor);
 DEFINE_LOG_CATEGORY(LogComposureEditor);
 
 class FComposureEditorModule : public IModuleInterface
 {
 	virtual void StartupModule() override
 	{
+		LLM_SCOPE_BYTAG(Compositing_ComposureEditor);
 		ISequencerModule& SequencerModule = FModuleManager::LoadModuleChecked<ISequencerModule>("Sequencer");
 		CreatePostMoveSettingsPropertyTrackEditorHandle = SequencerModule.RegisterPropertyTrackEditor<FComposurePostMoveSettingsPropertyTrackEditor>();
 		ComposureExportTrackEditorHandle               = SequencerModule.RegisterTrackEditor(FOnCreateTrackEditor::CreateLambda([](TSharedRef<ISequencer> In){ return MakeShared<FComposureExportTrackEditor>(In); }));
diff --git a/Engine/Plugins/Compositing/Composure/Source/ComposureLayersEditor/Private/CompElementEditorModule.cpp b/Engine/Plugins/Compositing/Composure/Source/ComposureLayersEditor/Private/CompElementEditorModule.cpp
--- a/Engine/Plugins/Compositing/Composure/Source/ComposureLayersEditor/Private/CompElementEditorModule.cpp
+++ b/Engine/Plugins/Compositing/Composure/Source/ComposureLayersEditor/Private/CompElementEditorModule.cpp
@@ -25,6 +25,8 @@
 #include "LevelEditorViewport.h"
 #include "Widgets/SCompElementPreviewDialog.h"
 
+LLM_DEFINE_TAG(Compositing_ComposureLayersEditor);
+
 namespace CompElementEditor_Impl
 {
 	static const FName ComposureLayersTabName(TEXT("ComposureLayers"));
@@ -90,6 +92,7 @@ private:
 //------------------------------------------------------------------------------
 void FCompElementEditorModule::StartupModule() 
 {
+	LLM_SCOPE_BYTAG(Compositing_ComposureLayersEditor);
 	FComposureEditorStyle::Get();
 	FCompElementEditorCommands::Register();
 
@@ -272,6 +275,7 @@ void FCompElementEditorModule::RegisterEditorTab()
 
 	LevelEditorTabManagerChangedHandle = LevelEditorModule.OnTabManagerChanged().AddLambda([]()
 	{
+		LLM_SCOPE_BYTAG(Compositing_ComposureLayersEditor);
 		FLevelEditorModule& LevelEditorModule = FModuleManager::GetModuleChecked<FLevelEditorModule>(CompElementEditor_Impl::LevelEditorModuleName);
 		TSharedPtr<FTabManager> LevelEditorTabManager = LevelEditorModule.GetLevelEditorTabManager();
 
diff --git a/Engine/Plugins/Compositing/OpenColorIO/Source/OpenColorIOEditor/Private/OpenColorIOEditorModule.cpp b/Engine/Plugins/Compositing/OpenColorIO/Source/OpenColorIOEditor/Private/OpenColorIOEditorModule.cpp
--- a/Engine/Plugins/Compositing/OpenColorIO/Source/OpenColorIOEditor/Private/OpenColorIOEditorModule.cpp
+++ b/Engine/Plugins/Compositing/OpenColorIO/Source/OpenColorIOEditor/Private/OpenColorIOEditorModule.cpp
@@ -23,7 +23,7 @@
 #include "ToolMenus.h"
 #include "Widgets/SOpenColorIODisplay.h"
 
-
+LLM_DEFINE_TAG(Compositing_OpenColorIOEditor);
 DEFINE_LOG_CATEGORY(LogOpenColorIOEditor);
 
 #define LOCTEXT_NAMESPACE "OpenColorIOEditorModule"
@@ -31,6 +31,7 @@ DEFINE_LOG_CATEGORY(LogOpenColorIOEditor);
 
 void FOpenColorIOEditorModule::StartupModule()
 {
+	LLM_SCOPE_BYTAG(Compositing_OpenColorIOEditor);
 	FWorldDelegates::OnPreWorldInitialization.AddRaw(this, &FOpenColorIOEditorModule::OnWorldInit);
 
 	// Register asset type actions for OpenColorIOConfiguration class
diff --git a/Engine/Plugins/Developer/Concert/ConcertApp/DisasterRecoveryClient/Source/DisasterRecoveryClient/Private/DisasterRecoveryClientModule.cpp b/Engine/Plugins/Developer/Concert/ConcertApp/DisasterRecoveryClient/Source/DisasterRecoveryClient/Private/DisasterRecoveryClientModule.cpp
--- a/Engine/Plugins/Developer/Concert/ConcertApp/DisasterRecoveryClient/Source/DisasterRecoveryClient/Private/DisasterRecoveryClientModule.cpp
+++ b/Engine/Plugins/Developer/Concert/ConcertApp/DisasterRecoveryClient/Source/DisasterRecoveryClient/Private/DisasterRecoveryClientModule.cpp
@@ -40,6 +40,7 @@
 	#include "SDisasterRecoveryHub.h"
 #endif
 
+LLM_DEFINE_TAG(Concert_DisasterRecoveryClient);
 DEFINE_LOG_CATEGORY(LogDisasterRecovery);
 
 #define LOCTEXT_NAMESPACE "DisasterRecoveryClient"
@@ -78,6 +79,8 @@ class FDisasterRecoveryClientModule : public IDisasterRecoveryClientModule
 public:
 	virtual void StartupModule() override
 	{
+		LLM_SCOPE_BYTAG(Concert_DisasterRecoveryClient);
+
 		if (!FApp::HasProjectName())
 		{
 			UE_LOG(LogDisasterRecovery, Warning, TEXT("The current project doesn't have any name. Recovery Hub will be disabled! Please set your project name."));
diff --git a/Engine/Plugins/Developer/Concert/ConcertApp/MultiUserClient/Source/MultiUserClient/Private/MultiUserClientModule.cpp b/Engine/Plugins/Developer/Concert/ConcertApp/MultiUserClient/Source/MultiUserClient/Private/MultiUserClientModule.cpp
--- a/Engine/Plugins/Developer/Concert/ConcertApp/MultiUserClient/Source/MultiUserClient/Private/MultiUserClientModule.cpp
+++ b/Engine/Plugins/Developer/Concert/ConcertApp/MultiUserClient/Source/MultiUserClient/Private/MultiUserClientModule.cpp
@@ -80,6 +80,8 @@ static const TCHAR MultiUserSlateServerAppName[] = TEXT("UnrealMultiUserSlateSer
 
 #define LOCTEXT_NAMESPACE "MultiUserClient"
 
+LLM_DEFINE_TAG(Concert_MultiUserClient);
+
 namespace MultiUserClientUtil
 {
 
@@ -593,6 +595,8 @@ class FMultiUserClientModule : public IMultiUserClientModule
 public:
 	virtual void StartupModule() override
 	{
+		LLM_SCOPE_BYTAG(Concert_MultiUserClient);
+
 		// Hook to the PreExit callback, needed to execute UObject related shutdowns
 		FCoreDelegates::OnPreExit.AddRaw(this, &FMultiUserClientModule::HandleAppPreExit);
 
diff --git a/Engine/Plugins/Developer/Concert/ConcertMain/Source/Concert/Private/ConcertClient.cpp b/Engine/Plugins/Developer/Concert/ConcertMain/Source/Concert/Private/ConcertClient.cpp
--- a/Engine/Plugins/Developer/Concert/ConcertMain/Source/Concert/Private/ConcertClient.cpp
+++ b/Engine/Plugins/Developer/Concert/ConcertMain/Source/Concert/Private/ConcertClient.cpp
@@ -22,6 +22,7 @@
 
 #define LOCTEXT_NAMESPACE "ConcertClient"
 
+LLM_DEFINE_TAG(Concert_ConcertClient);
 
 namespace ConcertUtil
 {
@@ -672,6 +673,7 @@ const FString& FConcertClient::GetRole() const
 
 void FConcertClient::Configure(const UConcertClientConfig* InSettings)
 {
+	LLM_SCOPE_BYTAG(Concert_ConcertClient);
 	ClientInfo.Initialize();
 	check(InSettings != nullptr);
 	Settings = TStrongObjectPtr<const UConcertClientConfig>(InSettings);
@@ -709,6 +711,7 @@ bool FConcertClient::IsStarted() const
 
 void FConcertClient::Startup()
 {
+	LLM_SCOPE_BYTAG(Concert_ConcertClient);
 	check(IsConfigured());
 	if (!ClientAdminEndpoint.IsValid() && EndpointProvider.IsValid())
 	{
diff --git a/Engine/Plugins/Developer/Concert/ConcertMain/Source/Concert/Private/ConcertMessageData.cpp b/Engine/Plugins/Developer/Concert/ConcertMain/Source/Concert/Private/ConcertMessageData.cpp
--- a/Engine/Plugins/Developer/Concert/ConcertMain/Source/Concert/Private/ConcertMessageData.cpp
+++ b/Engine/Plugins/Developer/Concert/ConcertMain/Source/Concert/Private/ConcertMessageData.cpp
@@ -12,6 +12,8 @@
 #include "Backends/CborStructSerializerBackend.h"
 #include "Backends/CborStructDeserializerBackend.h"
 
+LLM_DEFINE_TAG(Concert_ConcertClientInfo);
+
 void FConcertInstanceInfo::Initialize()
 {
 	InstanceId = FApp::GetInstanceId();
@@ -77,6 +79,7 @@ FText FConcertServerInfo::ToDisplayString() const
 
 void FConcertClientInfo::Initialize()
 {
+	LLM_SCOPE_BYTAG(Concert_ConcertClientInfo);
 	InstanceInfo.Initialize();
 	DeviceName = FPlatformProcess::ComputerName();
 	PlatformName = FPlatformProperties::PlatformName();
diff --git a/Engine/Plugins/Developer/Concert/ConcertMain/Source/ConcertTransport/Private/ConcertLocalEndpoint.cpp b/Engine/Plugins/Developer/Concert/ConcertMain/Source/ConcertTransport/Private/ConcertLocalEndpoint.cpp
--- a/Engine/Plugins/Developer/Concert/ConcertMain/Source/ConcertTransport/Private/ConcertLocalEndpoint.cpp
+++ b/Engine/Plugins/Developer/Concert/ConcertMain/Source/ConcertTransport/Private/ConcertLocalEndpoint.cpp
@@ -14,6 +14,8 @@
 #include "Misc/DateTime.h"
 #include "Stats/Stats.h"
 
+LLM_DEFINE_TAG(Concert_ConcertLocalEndpoint);
+
 class FConcertLocalEndpointKeepAliveRunnable : public FRunnable
 {
 public:
@@ -326,6 +328,7 @@ FConcertRemoteEndpointPtr FConcertLocalEndpoint::FindRemoteEndpoint(const FMessa
 
 bool FConcertLocalEndpoint::HandleTick(float DeltaTime)
 {
+	LLM_SCOPE_BYTAG(Concert_ConcertLocalEndpoint);
 	SCOPED_CONCERT_TRACE(FConcertLocalEndpoint_HandleTick);
 
 
diff --git a/Engine/Plugins/Developer/Concert/ConcertSync/ConcertSyncClient/Source/ConcertSyncClient/Private/ConcertSyncClient.cpp b/Engine/Plugins/Developer/Concert/ConcertSync/ConcertSyncClient/Source/ConcertSyncClient/Private/ConcertSyncClient.cpp
--- a/Engine/Plugins/Developer/Concert/ConcertSync/ConcertSyncClient/Source/ConcertSyncClient/Private/ConcertSyncClient.cpp
+++ b/Engine/Plugins/Developer/Concert/ConcertSync/ConcertSyncClient/Source/ConcertSyncClient/Private/ConcertSyncClient.cpp
@@ -12,6 +12,7 @@
 #include "ConcertSourceControlProxy.h"
 
 #define LOCTEXT_NAMESPACE "ConcertSyncClient"
+LLM_DEFINE_TAG(Concert_ConcertSyncClient);
 
 FConcertSyncClient::FConcertSyncClient(const FString& InRole, IConcertClientPackageBridge* InPackageBridge, IConcertClientTransactionBridge* InTransactionBridge)
 	: ConcertClient(IConcertModule::Get().CreateClient(InRole))
@@ -38,6 +39,7 @@ FConcertSyncClient::~FConcertSyncClient()
 
 void FConcertSyncClient::Startup(const UConcertClientConfig* InClientConfig, const EConcertSyncSessionFlags InSessionFlags)
 {
+	LLM_SCOPE_BYTAG(Concert_ConcertSyncClient);
 	SessionFlags = InSessionFlags;
 
 	// Boot the client instance
diff --git a/Engine/Plugins/Developer/Concert/ConcertUI/ConcertSharedSlate/Source/ConcertSharedSlate/Private/ConcertFrontendStyle.cpp b/Engine/Plugins/Developer/Concert/ConcertUI/ConcertSharedSlate/Source/ConcertSharedSlate/Private/ConcertFrontendStyle.cpp
--- a/Engine/Plugins/Developer/Concert/ConcertUI/ConcertSharedSlate/Source/ConcertSharedSlate/Private/ConcertFrontendStyle.cpp
+++ b/Engine/Plugins/Developer/Concert/ConcertUI/ConcertSharedSlate/Source/ConcertSharedSlate/Private/ConcertFrontendStyle.cpp
@@ -8,6 +8,8 @@
 #include "SlateOptMacros.h"
 #include "Styling/StarshipCoreStyle.h"
 
+LLM_DEFINE_TAG(Concert_ConcertFrontendStyle);
+
 #define IMAGE_PLUGIN_BRUSH( RelativePath, ... ) FSlateImageBrush( FConcertFrontendStyle::InContent( RelativePath, ".png" ), __VA_ARGS__ )
 #define IMAGE_BRUSH(RelativePath, ...) FSlateImageBrush(StyleSet->RootToContentDir(RelativePath, TEXT(".png")), __VA_ARGS__)
 #define BOX_BRUSH(RelativePath, ... ) FSlateBoxBrush( StyleSet->RootToContentDir(RelativePath, TEXT(".png") ), __VA_ARGS__ )
@@ -28,6 +30,7 @@ FName FConcertFrontendStyle::GetStyleSetName()
 
 void FConcertFrontendStyle::Initialize()
 {
+	LLM_SCOPE_BYTAG(Concert_ConcertFrontendStyle);
 	// Only register once
 	if (StyleSet.IsValid())
 	{
diff --git a/Engine/Plugins/Experimental/RemoteSession/Source/RemoteSessionEditor/Private/RemoteSessionEditorModule.cpp b/Engine/Plugins/Experimental/RemoteSession/Source/RemoteSessionEditor/Private/RemoteSessionEditorModule.cpp
--- a/Engine/Plugins/Experimental/RemoteSession/Source/RemoteSessionEditor/Private/RemoteSessionEditorModule.cpp
+++ b/Engine/Plugins/Experimental/RemoteSession/Source/RemoteSessionEditor/Private/RemoteSessionEditorModule.cpp
@@ -14,6 +14,7 @@ class FRemoteSessionEditorModule : public IModuleInterface
 public:
 	virtual void StartupModule() override
 	{
+		LLM_SCOPE_BYNAME(TEXT("RemoteSessionEditor"));
 		FRemoteSessionEditorStyle::Register();
 		SRemoteSessionStream::RegisterNomadTabSpawner();
 	}
diff --git a/Engine/Plugins/Experimental/VirtualCamera/Source/VirtualCamera/Private/VirtualCamera.cpp b/Engine/Plugins/Experimental/VirtualCamera/Source/VirtualCamera/Private/VirtualCamera.cpp
--- a/Engine/Plugins/Experimental/VirtualCamera/Source/VirtualCamera/Private/VirtualCamera.cpp
+++ b/Engine/Plugins/Experimental/VirtualCamera/Source/VirtualCamera/Private/VirtualCamera.cpp
@@ -30,6 +30,7 @@ public:
 private:
 	virtual void StartupModule() override
 	{
+		LLM_SCOPE_BYNAME(TEXT("VirtualCamera"));
 		ConcertManager = MakeUnique<FConcertVirtualCameraManager>();
 
 		// Loads widgets (ex. RadialSlider) that are potentially referenced by assets
diff --git a/Engine/Plugins/Experimental/VirtualProductionUtilities/Source/VPMaterialsEditor/Private/VPMaterialsEditorModule.cpp b/Engine/Plugins/Experimental/VirtualProductionUtilities/Source/VPMaterialsEditor/Private/VPMaterialsEditorModule.cpp
--- a/Engine/Plugins/Experimental/VirtualProductionUtilities/Source/VPMaterialsEditor/Private/VPMaterialsEditorModule.cpp
+++ b/Engine/Plugins/Experimental/VirtualProductionUtilities/Source/VPMaterialsEditor/Private/VPMaterialsEditorModule.cpp
@@ -10,6 +10,7 @@ void FVPMaterialsEditorModule::StartupModule()
 	// Add bottom extender for material item
 	FMaterialList::OnAddMaterialItemViewExtraBottomWidget.AddLambda([](const TSharedRef<FMaterialItemView>& InMaterialItemView, UActorComponent* InCurrentComponent, IDetailLayoutBuilder& InDetailBuilder, TArray<TSharedPtr<SWidget>>& OutExtensions)
 	{
+		LLM_SCOPE_BYNAME("VirtualProduction/VPMaterialsEditor");
 		OutExtensions.Add(SNew(SMaterialDynamicView, InMaterialItemView, InCurrentComponent));
 	});
 }
diff --git a/Engine/Plugins/Media/TimecodeSynchronizer/Source/TimecodeSynchronizerEditor/Private/TimecodeSynchronizerEditorModule.cpp b/Engine/Plugins/Media/TimecodeSynchronizer/Source/TimecodeSynchronizerEditor/Private/TimecodeSynchronizerEditorModule.cpp
--- a/Engine/Plugins/Media/TimecodeSynchronizer/Source/TimecodeSynchronizerEditor/Private/TimecodeSynchronizerEditorModule.cpp
+++ b/Engine/Plugins/Media/TimecodeSynchronizer/Source/TimecodeSynchronizerEditor/Private/TimecodeSynchronizerEditorModule.cpp
@@ -22,6 +22,7 @@ class FTimecodeSynchronizerEditorModule : public ITimecodeSynchronizerEditorModu
 public:
 	virtual void StartupModule() override
 	{
+		LLM_SCOPE_BYNAME(TEXT("TimecodeSynchronizerEditor"));
 		// Disable any UI feature if running in command mode
 		if (!IsRunningCommandlet())
 		{
diff --git a/Engine/Plugins/VirtualProduction/MultiUserTakes/Source/ConcertTakeRecorder/Private/ConcertTakeRecorderModule.cpp b/Engine/Plugins/VirtualProduction/MultiUserTakes/Source/ConcertTakeRecorder/Private/ConcertTakeRecorderModule.cpp
--- a/Engine/Plugins/VirtualProduction/MultiUserTakes/Source/ConcertTakeRecorder/Private/ConcertTakeRecorderModule.cpp
+++ b/Engine/Plugins/VirtualProduction/MultiUserTakes/Source/ConcertTakeRecorder/Private/ConcertTakeRecorderModule.cpp
@@ -6,6 +6,8 @@
 #include "ConcertTakeRecorderManager.h"
 #include "ConcertTakeRecorderStyle.h"
 
+LLM_DEFINE_TAG(Concert_ConcertTakeRecorder);
+
 /**
  * Module that adds multi user synchronization to take recorder.
  */
@@ -14,6 +16,7 @@ class FConcertTakeRecorderModule : public IModuleInterface
 private:
 	virtual void StartupModule() override
 	{
+		LLM_SCOPE_BYTAG(Concert_ConcertTakeRecorder);
 		FConcertTakeRecorderStyle::Initialize();
 		ConcertManager = MakeUnique<FConcertTakeRecorderManager>();
 	}
diff --git a/Engine/Source/Editor/PixelInspector/Private/PixelInspectorModule.cpp b/Engine/Source/Editor/PixelInspector/Private/PixelInspectorModule.cpp
--- a/Engine/Source/Editor/PixelInspector/Private/PixelInspectorModule.cpp
+++ b/Engine/Source/Editor/PixelInspector/Private/PixelInspectorModule.cpp
@@ -123,6 +123,7 @@ void FPixelInspectorModule::OnTabClosed(TSharedRef<SDockTab> TabBeingClosed)
 
 void FPixelInspectorModule::RegisterTabSpawner(const TSharedPtr<FWorkspaceItem>& WorkspaceGroup)
 {
+	LLM_SCOPE_BYNAME(TEXT("PixelInspector"));
 	if (bHasRegisteredTabSpawners)
 	{
 		UnregisterTabSpawner();
diff --git a/Engine/Source/Editor/VREditor/VREditorModule.cpp b/Engine/Source/Editor/VREditor/VREditorModule.cpp
--- a/Engine/Source/Editor/VREditor/VREditorModule.cpp
+++ b/Engine/Source/Editor/VREditor/VREditorModule.cpp
@@ -67,6 +67,7 @@ namespace VREd
 
 void FVREditorModule::StartupModule()
 {
+	LLM_SCOPE_BYNAME(TEXT("VREditor"));
 	RadialMenuExtender = MakeShareable(new FExtender());
 }
 
