ObjectPtr upgrade for engine plugins used from UnrealEditor -allmodules.  Automated upgrade process.

diff --git a/Engine/Plugins/AI/MLAdapter/Source/MLAdapter/Public/Agents/MLAdapterAgent.h b/Engine/Plugins/AI/MLAdapter/Source/MLAdapter/Public/Agents/MLAdapterAgent.h
--- a/Engine/Plugins/AI/MLAdapter/Source/MLAdapter/Public/Agents/MLAdapterAgent.h
+++ b/Engine/Plugins/AI/MLAdapter/Source/MLAdapter/Public/Agents/MLAdapterAgent.h
@@ -198,10 +198,10 @@ protected:
 
 protected:
 	UPROPERTY(EditDefaultsOnly, Instanced, BlueprintReadOnly, Category = MLAdapter)
-	TArray<UMLAdapterSensor*> Sensors;
+	TArray<TObjectPtr<UMLAdapterSensor>> Sensors;
 
 	UPROPERTY(EditDefaultsOnly, Instanced, BlueprintReadOnly, Category = MLAdapter)
-	TArray<UMLAdapterActuator*> Actuators;
+	TArray<TObjectPtr<UMLAdapterActuator>> Actuators;
 
 	UPROPERTY(EditDefaultsOnly, BlueprintReadOnly, Category = MLAdapter)
 	TSubclassOf<AActor> AvatarClass;
@@ -225,13 +225,13 @@ protected:
 
 private:
 	UPROPERTY()
-	AActor* Avatar;
+	TObjectPtr<AActor> Avatar;
 
 	UPROPERTY()
-	AController* Controller;
+	TObjectPtr<AController> Controller;
 
 	UPROPERTY()
-	APawn* Pawn;
+	TObjectPtr<APawn> Pawn;
 
 	FMLAdapter::FAgentID AgentID;
 
diff --git a/Engine/Plugins/AI/MLAdapter/Source/MLAdapter/Public/Agents/MLAdapterAgent_Inference.h b/Engine/Plugins/AI/MLAdapter/Source/MLAdapter/Public/Agents/MLAdapterAgent_Inference.h
--- a/Engine/Plugins/AI/MLAdapter/Source/MLAdapter/Public/Agents/MLAdapterAgent_Inference.h
+++ b/Engine/Plugins/AI/MLAdapter/Source/MLAdapter/Public/Agents/MLAdapterAgent_Inference.h
@@ -29,5 +29,5 @@ public:
 
 	/** The neural network controlling this agent. */
 	UPROPERTY()
-	UNeuralNetwork* Brain;
+	TObjectPtr<UNeuralNetwork> Brain;
 };
diff --git a/Engine/Plugins/AI/MLAdapter/Source/MLAdapter/Public/Managers/MLAdapterManager.h b/Engine/Plugins/AI/MLAdapter/Source/MLAdapter/Public/Managers/MLAdapterManager.h
--- a/Engine/Plugins/AI/MLAdapter/Source/MLAdapter/Public/Managers/MLAdapterManager.h
+++ b/Engine/Plugins/AI/MLAdapter/Source/MLAdapter/Public/Managers/MLAdapterManager.h
@@ -161,10 +161,10 @@ protected:
 	friend struct FMLAdapterConsoleCommands::FHelper;
 
 	UPROPERTY()
-	UMLAdapterSession* Session;
+	TObjectPtr<UMLAdapterSession> Session;
 
 	UPROPERTY()
-	UWorld* LastActiveWorld;
+	TObjectPtr<UWorld> LastActiveWorld;
 
 	UPROPERTY()
 	FMLAdapterLibrarian Librarian;
diff --git a/Engine/Plugins/AI/MLAdapter/Source/MLAdapter/Public/Sensors/MLAdapterSensor_AIPerception.h b/Engine/Plugins/AI/MLAdapter/Source/MLAdapter/Public/Sensors/MLAdapterSensor_AIPerception.h
--- a/Engine/Plugins/AI/MLAdapter/Source/MLAdapter/Public/Sensors/MLAdapterSensor_AIPerception.h
+++ b/Engine/Plugins/AI/MLAdapter/Source/MLAdapter/Public/Sensors/MLAdapterSensor_AIPerception.h
@@ -74,7 +74,7 @@ protected:
 
 protected:
 	UPROPERTY()
-	UAIPerceptionComponent* PerceptionComponent;
+	TObjectPtr<UAIPerceptionComponent> PerceptionComponent;
 
 	/**
 	 *	When set to true will only gather perception "delta" meaning consecutive
diff --git a/Engine/Plugins/AI/MLAdapter/Source/MLAdapter/Public/Sensors/MLAdapterSensor_Attribute.h b/Engine/Plugins/AI/MLAdapter/Source/MLAdapter/Public/Sensors/MLAdapterSensor_Attribute.h
--- a/Engine/Plugins/AI/MLAdapter/Source/MLAdapter/Public/Sensors/MLAdapterSensor_Attribute.h
+++ b/Engine/Plugins/AI/MLAdapter/Source/MLAdapter/Public/Sensors/MLAdapterSensor_Attribute.h
@@ -32,7 +32,7 @@ protected:
 	void BindAttributes(AActor& Actor);
 protected:
 	UPROPERTY()
-	UAttributeSet* AttributeSet;
+	TObjectPtr<UAttributeSet> AttributeSet;
 
 	TArray<FName> AttributeNames;
 	// valid only as long as AttributeSet != nullptr
diff --git a/Engine/Plugins/AI/MLAdapter/Source/MLAdapter/Public/Sensors/MLAdapterSensor_Camera.h b/Engine/Plugins/AI/MLAdapter/Source/MLAdapter/Public/Sensors/MLAdapterSensor_Camera.h
--- a/Engine/Plugins/AI/MLAdapter/Source/MLAdapter/Public/Sensors/MLAdapterSensor_Camera.h
+++ b/Engine/Plugins/AI/MLAdapter/Source/MLAdapter/Public/Sensors/MLAdapterSensor_Camera.h
@@ -44,13 +44,13 @@ protected:
 	uint32 bShowUI : 1;
 
 	UPROPERTY(Transient)
-	USceneCaptureComponent2D* CaptureComp;
+	TObjectPtr<USceneCaptureComponent2D> CaptureComp;
 
 	UPROPERTY(Transient)
-	UTextureRenderTarget2D* RenderTarget2D;
+	TObjectPtr<UTextureRenderTarget2D> RenderTarget2D;
 
 	UPROPERTY()
-	UGameViewportClient* CachedViewportClient;
+	TObjectPtr<UGameViewportClient> CachedViewportClient;
 
 	ETextureRenderTargetFormat RenderTargetFormat = ETextureRenderTargetFormat::RTF_RGBA16f;
 	ESceneCaptureSource CaptureSource = SCS_FinalColorLDR;
diff --git a/Engine/Plugins/AI/MLAdapter/Source/MLAdapter/Public/Sensors/MLAdapterSensor_Input.h b/Engine/Plugins/AI/MLAdapter/Source/MLAdapter/Public/Sensors/MLAdapterSensor_Input.h
--- a/Engine/Plugins/AI/MLAdapter/Source/MLAdapter/Public/Sensors/MLAdapterSensor_Input.h
+++ b/Engine/Plugins/AI/MLAdapter/Source/MLAdapter/Public/Sensors/MLAdapterSensor_Input.h
@@ -36,7 +36,7 @@ protected:
 	virtual void OnAvatarSet(AActor* Avatar);
 
 	UPROPERTY(EditDefaultsOnly, Category = MLAdapter)
-	UGameViewportClient* GameViewport;
+	TObjectPtr<UGameViewportClient> GameViewport;
 
 	UPROPERTY(EditDefaultsOnly, Category = MLAdapter)
 	uint32 bRecordKeyRelease : 1;
diff --git a/Engine/Plugins/AI/MLAdapter/Source/MLAdapter/Public/Sessions/MLAdapterSession.h b/Engine/Plugins/AI/MLAdapter/Source/MLAdapter/Public/Sessions/MLAdapterSession.h
--- a/Engine/Plugins/AI/MLAdapter/Source/MLAdapter/Public/Sessions/MLAdapterSession.h
+++ b/Engine/Plugins/AI/MLAdapter/Source/MLAdapter/Public/Sessions/MLAdapterSession.h
@@ -187,20 +187,20 @@ protected:
 	}
 
 	UPROPERTY()
-	AGameModeBase* CachedGameMode;
+	TObjectPtr<AGameModeBase> CachedGameMode;
 
 	UPROPERTY()
-	UWorld* CachedWorld;
+	TObjectPtr<UWorld> CachedWorld;
 
 	/** @see HashAvatar */
 	UPROPERTY()
-	TMap<uint32, UMLAdapterAgent*> AvatarToAgent;
+	TMap<uint32, TObjectPtr<UMLAdapterAgent>> AvatarToAgent;
 
 	UPROPERTY()
-	TArray<UMLAdapterAgent*> Agents;
+	TArray<TObjectPtr<UMLAdapterAgent>> Agents;
 
 	UPROPERTY()
-	TArray<UMLAdapterAgent*> AwaitingAvatar;
+	TArray<TObjectPtr<UMLAdapterAgent>> AwaitingAvatar;
 
 	FOnAgentAvatarChangedDelegate OnAgentAvatarChanged;
 	FOnBeginAgentRemove OnBeginAgentRemove;
diff --git a/Engine/Plugins/AI/MassAI/Source/MassAIBehavior/Public/MassComponentHitSubsystem.h b/Engine/Plugins/AI/MassAI/Source/MassAIBehavior/Public/MassComponentHitSubsystem.h
--- a/Engine/Plugins/AI/MassAI/Source/MassAIBehavior/Public/MassComponentHitSubsystem.h
+++ b/Engine/Plugins/AI/MassAI/Source/MassAIBehavior/Public/MassComponentHitSubsystem.h
@@ -38,17 +38,17 @@ protected:
 	void OnHitCallback(UPrimitiveComponent* HitComp, AActor* OtherActor, UPrimitiveComponent* OtherComp, FVector NormalImpulse, const FHitResult& Hit);
 
 	UPROPERTY()
-	UMassSignalSubsystem* SignalSubsystem;
+	TObjectPtr<UMassSignalSubsystem> SignalSubsystem;
 
 	UPROPERTY()
-    UMassAgentSubsystem* AgentSubsystem;
+    TObjectPtr<UMassAgentSubsystem> AgentSubsystem;
 
 	UPROPERTY()
 	TMap<FMassEntityHandle, FMassHitResult> HitResults;
 
 	UPROPERTY()
-	TMap<UActorComponent*, FMassEntityHandle> ComponentToEntityMap;
+	TMap<TObjectPtr<UActorComponent>, FMassEntityHandle> ComponentToEntityMap;
 
 	UPROPERTY()
-	TMap<FMassEntityHandle, UActorComponent*> EntityToComponentMap;
+	TMap<FMassEntityHandle, TObjectPtr<UActorComponent>> EntityToComponentMap;
 };
diff --git a/Engine/Plugins/AI/MassAI/Source/MassAIBehavior/Public/MassStateTreeFragments.h b/Engine/Plugins/AI/MassAI/Source/MassAIBehavior/Public/MassStateTreeFragments.h
--- a/Engine/Plugins/AI/MassAI/Source/MassAIBehavior/Public/MassStateTreeFragments.h
+++ b/Engine/Plugins/AI/MassAI/Source/MassAIBehavior/Public/MassStateTreeFragments.h
@@ -30,5 +30,5 @@ struct MASSAIBEHAVIOR_API FMassStateTreeSharedFragment : public FMassSharedFragm
 	FMassStateTreeSharedFragment() = default;
 
 	UPROPERTY()
-	UStateTree* StateTree = nullptr;
+	TObjectPtr<UStateTree> StateTree = nullptr;
 };
diff --git a/Engine/Plugins/AI/MassAI/Source/MassAIBehavior/Public/MassStateTreeProcessors.h b/Engine/Plugins/AI/MassAI/Source/MassAIBehavior/Public/MassStateTreeProcessors.h
--- a/Engine/Plugins/AI/MassAI/Source/MassAIBehavior/Public/MassStateTreeProcessors.h
+++ b/Engine/Plugins/AI/MassAI/Source/MassAIBehavior/Public/MassStateTreeProcessors.h
@@ -29,7 +29,7 @@ protected:
 	FMassEntityQuery EntityQuery;
 
 	UPROPERTY(Transient)
-	UMassSignalSubsystem* SignalSubsystem = nullptr;
+	TObjectPtr<UMassSignalSubsystem> SignalSubsystem = nullptr;
 };
 
 /**
diff --git a/Engine/Plugins/AI/MassAI/Source/MassAIBehavior/Public/MassStateTreeTrait.h b/Engine/Plugins/AI/MassAI/Source/MassAIBehavior/Public/MassStateTreeTrait.h
--- a/Engine/Plugins/AI/MassAI/Source/MassAIBehavior/Public/MassStateTreeTrait.h
+++ b/Engine/Plugins/AI/MassAI/Source/MassAIBehavior/Public/MassStateTreeTrait.h
@@ -22,5 +22,5 @@ protected:
 	virtual void ValidateTemplate(FMassEntityTemplateBuildContext& BuildContext, const UWorld& World) const override;
 
 	UPROPERTY(Category="StateTree", EditAnywhere, meta=(RequiredAssetDataTags="Schema=MassStateTreeSchema"))
-	UStateTree* StateTree;
+	TObjectPtr<UStateTree> StateTree;
 };
diff --git a/Engine/Plugins/AI/MassAI/Source/MassNavigationEditor/Private/MassNavigationTestingActor.h b/Engine/Plugins/AI/MassAI/Source/MassNavigationEditor/Private/MassNavigationTestingActor.h
--- a/Engine/Plugins/AI/MassAI/Source/MassNavigationEditor/Private/MassNavigationTestingActor.h
+++ b/Engine/Plugins/AI/MassAI/Source/MassNavigationEditor/Private/MassNavigationTestingActor.h
@@ -65,7 +65,7 @@ protected:
 	FDelegateHandle OnDataRemovedHandle;
 
 	UPROPERTY(Transient)
-	UZoneGraphSubsystem* ZoneGraph;
+	TObjectPtr<UZoneGraphSubsystem> ZoneGraph;
 
 	UPROPERTY(Transient)
 	FZoneGraphLaneLocation LaneLocation;
@@ -117,5 +117,5 @@ public:
 
 protected:
 	UPROPERTY(Category = Default, VisibleAnywhere, meta = (AllowPrivateAccess = "true"))
-	UMassNavigationTestingComponent* DebugComp;
+	TObjectPtr<UMassNavigationTestingComponent> DebugComp;
 };
diff --git a/Engine/Plugins/AI/MassAI/Source/MassZoneGraphNavigation/Public/MassZoneGraphNavigationProcessors.h b/Engine/Plugins/AI/MassAI/Source/MassZoneGraphNavigation/Public/MassZoneGraphNavigationProcessors.h
--- a/Engine/Plugins/AI/MassAI/Source/MassZoneGraphNavigation/Public/MassZoneGraphNavigationProcessors.h
+++ b/Engine/Plugins/AI/MassAI/Source/MassZoneGraphNavigation/Public/MassZoneGraphNavigationProcessors.h
@@ -45,7 +45,7 @@ protected:
 	FMassEntityQuery EntityQuery_Conditional;
 
 	UPROPERTY(Transient)
-	UMassSignalSubsystem* SignalSubsystem = nullptr;
+	TObjectPtr<UMassSignalSubsystem> SignalSubsystem = nullptr;
 };
 
 /** ZoneGraph lane cache boundary processor */
diff --git a/Engine/Plugins/AI/MassCrowd/Source/MassCrowd/Public/MassCrowdLaneClosingTest.h b/Engine/Plugins/AI/MassCrowd/Source/MassCrowd/Public/MassCrowdLaneClosingTest.h
--- a/Engine/Plugins/AI/MassCrowd/Source/MassCrowd/Public/MassCrowdLaneClosingTest.h
+++ b/Engine/Plugins/AI/MassCrowd/Source/MassCrowd/Public/MassCrowdLaneClosingTest.h
@@ -23,7 +23,7 @@ private:
 	FZoneGraphLaneLocation LaneLocation;
 
 	UPROPERTY(Transient)
-	UMassCrowdSubsystem* CrowdSubsystem;
+	TObjectPtr<UMassCrowdSubsystem> CrowdSubsystem;
 
 	UPROPERTY(EditAnywhere, Category = Test)
 	ECrowdLaneState LaneState = ECrowdLaneState::Closed;
diff --git a/Engine/Plugins/AI/MassCrowd/Source/MassCrowd/Public/MassCrowdNavigationProcessor.h b/Engine/Plugins/AI/MassCrowd/Source/MassCrowd/Public/MassCrowdNavigationProcessor.h
--- a/Engine/Plugins/AI/MassCrowd/Source/MassCrowd/Public/MassCrowdNavigationProcessor.h
+++ b/Engine/Plugins/AI/MassCrowd/Source/MassCrowd/Public/MassCrowdNavigationProcessor.h
@@ -60,7 +60,7 @@ protected:
 	FMassEntityQuery EntityQuery_Conditional;
 
 	UPROPERTY(Transient)
-	UZoneGraphAnnotationSubsystem* ZoneGraphAnnotationSubsystem;
+	TObjectPtr<UZoneGraphAnnotationSubsystem> ZoneGraphAnnotationSubsystem;
 };
 
 
@@ -94,7 +94,7 @@ protected:
 	virtual void Execute(FMassEntityManager& EntityManager, FMassExecutionContext& Context) override;
 
 	UPROPERTY(Transient)
-	UZoneGraphAnnotationSubsystem* ZoneGraphAnnotationSubsystem;
+	TObjectPtr<UZoneGraphAnnotationSubsystem> ZoneGraphAnnotationSubsystem;
 
 	FMassEntityQuery EntityQuery;
 };
diff --git a/Engine/Plugins/AI/MassCrowd/Source/MassCrowd/Public/MassCrowdSubsystem.h b/Engine/Plugins/AI/MassCrowd/Source/MassCrowd/Public/MassCrowdSubsystem.h
--- a/Engine/Plugins/AI/MassCrowd/Source/MassCrowd/Public/MassCrowdSubsystem.h
+++ b/Engine/Plugins/AI/MassCrowd/Source/MassCrowd/Public/MassCrowdSubsystem.h
@@ -222,13 +222,13 @@ protected:
 	void CreateWaitSlots(const int32 CrossingLaneIndex, FCrowdWaitAreaData& WaitArea, const FZoneGraphStorage& ZoneGraphStorage);
 
 	UPROPERTY(Transient)
-	UZoneGraphSubsystem* ZoneGraphSubsystem = nullptr;
+	TObjectPtr<UZoneGraphSubsystem> ZoneGraphSubsystem = nullptr;
 
 	UPROPERTY(Transient)
-	UZoneGraphAnnotationSubsystem* ZoneGraphAnnotationSubsystem = nullptr;
+	TObjectPtr<UZoneGraphAnnotationSubsystem> ZoneGraphAnnotationSubsystem = nullptr;
 
 	UPROPERTY(Transient)
-	const UMassCrowdSettings* MassCrowdSettings = nullptr;
+	TObjectPtr<const UMassCrowdSettings> MassCrowdSettings = nullptr;
 
 	FDelegateHandle OnPostZoneGraphDataAddedHandle;
 	FDelegateHandle OnPreZoneGraphDataRemovedHandle;
diff --git a/Engine/Plugins/AI/MassCrowd/Source/MassCrowd/Public/MassCrowdVisualizationProcessor.h b/Engine/Plugins/AI/MassCrowd/Source/MassCrowd/Public/MassCrowdVisualizationProcessor.h
--- a/Engine/Plugins/AI/MassCrowd/Source/MassCrowd/Public/MassCrowdVisualizationProcessor.h
+++ b/Engine/Plugins/AI/MassCrowd/Source/MassCrowd/Public/MassCrowdVisualizationProcessor.h
@@ -47,7 +47,7 @@ protected:
 	virtual void Execute(FMassEntityManager& EntityManager, FMassExecutionContext& Context) override;
 
 	UPROPERTY(Transient)
-	UWorld* World;
+	TObjectPtr<UWorld> World;
 
 	FMassEntityQuery EntityQuery;
 };
\ No newline at end of file
diff --git a/Engine/Plugins/AI/MassCrowd/Source/MassCrowd/Public/ZoneGraphCrowdLaneAnnotations.h b/Engine/Plugins/AI/MassCrowd/Source/MassCrowd/Public/ZoneGraphCrowdLaneAnnotations.h
--- a/Engine/Plugins/AI/MassCrowd/Source/MassCrowd/Public/ZoneGraphCrowdLaneAnnotations.h
+++ b/Engine/Plugins/AI/MassCrowd/Source/MassCrowd/Public/ZoneGraphCrowdLaneAnnotations.h
@@ -62,5 +62,5 @@ protected:
 
 	/** Cached ZoneGraphSubsystem */
 	UPROPERTY(Transient)
-	UMassCrowdSubsystem* CrowdSubsystem = nullptr;
+	TObjectPtr<UMassCrowdSubsystem> CrowdSubsystem = nullptr;
 };
diff --git a/Engine/Plugins/Animation/AnimToTexture/Source/AnimToTexture/Public/AnimToTextureDataAsset.h b/Engine/Plugins/Animation/AnimToTexture/Source/AnimToTexture/Public/AnimToTextureDataAsset.h
--- a/Engine/Plugins/Animation/AnimToTexture/Source/AnimToTexture/Public/AnimToTextureDataAsset.h
+++ b/Engine/Plugins/Animation/AnimToTexture/Source/AnimToTexture/Public/AnimToTextureDataAsset.h
@@ -133,7 +133,7 @@ struct FAnimSequenceInfo
 	bool bEnabled = true;
 
 	UPROPERTY(EditAnywhere, Category = Default, BlueprintReadWrite)
-	UAnimSequence* AnimSequence = nullptr;
+	TObjectPtr<UAnimSequence> AnimSequence = nullptr;
 
 	UPROPERTY(EditAnywhere, Category = Default, BlueprintReadWrite)
 	bool bLooping = true;
diff --git a/Engine/Plugins/Animation/AnimToTexture/Source/AnimToTexture/Public/EvaluateSequenceAnimInstance.h b/Engine/Plugins/Animation/AnimToTexture/Source/AnimToTexture/Public/EvaluateSequenceAnimInstance.h
--- a/Engine/Plugins/Animation/AnimToTexture/Source/AnimToTexture/Public/EvaluateSequenceAnimInstance.h
+++ b/Engine/Plugins/Animation/AnimToTexture/Source/AnimToTexture/Public/EvaluateSequenceAnimInstance.h
@@ -19,7 +19,7 @@ public:
 	UEvaluateSequenceAnimInstance(const FObjectInitializer& ObjectInitializer = FObjectInitializer::Get());
 
 	UPROPERTY(Transient, EditAnywhere, Category = Default, BlueprintReadOnly)
-	UAnimSequence* SequenceToEvaluate = nullptr;
+	TObjectPtr<UAnimSequence> SequenceToEvaluate = nullptr;
 
 	UPROPERTY(Transient, EditAnywhere, Category = Default, BlueprintReadOnly)
 	float TimeToEvaluate = 0.0f;
diff --git a/Engine/Plugins/Animation/RigLogic/Source/RigLogicEditor/Public/DNAAssetImportFactory.h b/Engine/Plugins/Animation/RigLogic/Source/RigLogicEditor/Public/DNAAssetImportFactory.h
--- a/Engine/Plugins/Animation/RigLogic/Source/RigLogicEditor/Public/DNAAssetImportFactory.h
+++ b/Engine/Plugins/Animation/RigLogic/Source/RigLogicEditor/Public/DNAAssetImportFactory.h
@@ -20,11 +20,11 @@ class UDNAAssetImportFactory: public UFactory
 
 public:
 	UPROPERTY()
-		class UDNAAssetImportUI* ImportUI;
+		TObjectPtr<class UDNAAssetImportUI> ImportUI;
 
 	/** Prevent garbage collection of original when overriding ImportUI property */
 	UPROPERTY()
-		class UDNAAssetImportUI* OriginalImportUI;
+		TObjectPtr<class UDNAAssetImportUI> OriginalImportUI;
 
 	/** UObject properties */
 	virtual void PostInitProperties() override;
diff --git a/Engine/Plugins/Animation/RigLogic/Source/RigLogicEditor/Public/DNAAssetImportUI.h b/Engine/Plugins/Animation/RigLogic/Source/RigLogicEditor/Public/DNAAssetImportUI.h
--- a/Engine/Plugins/Animation/RigLogic/Source/RigLogicEditor/Public/DNAAssetImportUI.h
+++ b/Engine/Plugins/Animation/RigLogic/Source/RigLogicEditor/Public/DNAAssetImportUI.h
@@ -22,7 +22,7 @@ public:
 
 	/** Skeletal mesh to use for imported DNA asset. When importing DNA, leaving this as "None" will generate new skeletal mesh. */
 	UPROPERTY(EditAnywhere, BlueprintReadWrite, Category = Mesh, meta = (ImportType = "DNAAsset"))
-		class USkeletalMesh* SkeletalMesh;
+		TObjectPtr<class USkeletalMesh> SkeletalMesh;
 
 	UFUNCTION(BlueprintCallable, Category = Miscellaneous)
 		void ResetToDefault();
diff --git a/Engine/Plugins/AudioGameplayVolume/Source/AudioGameplayVolume/Public/AudioGameplayVolume.h b/Engine/Plugins/AudioGameplayVolume/Source/AudioGameplayVolume/Public/AudioGameplayVolume.h
--- a/Engine/Plugins/AudioGameplayVolume/Source/AudioGameplayVolume/Public/AudioGameplayVolume.h
+++ b/Engine/Plugins/AudioGameplayVolume/Source/AudioGameplayVolume/Public/AudioGameplayVolume.h
@@ -28,7 +28,7 @@ private:
 
 	// A representation of this volume for the audio thread
 	UPROPERTY()
-	UAudioGameplayVolumeComponent* AGVComponent = nullptr;
+	TObjectPtr<UAudioGameplayVolumeComponent> AGVComponent = nullptr;
 
 	// Whether this volume is currently enabled.  Disabled volumes will not have a volume proxy, 
 	// and therefore will not be considered for intersection checks.
diff --git a/Engine/Plugins/AudioGameplayVolume/Source/AudioGameplayVolume/Public/AudioGameplayVolumeComponent.h b/Engine/Plugins/AudioGameplayVolume/Source/AudioGameplayVolume/Public/AudioGameplayVolumeComponent.h
--- a/Engine/Plugins/AudioGameplayVolume/Source/AudioGameplayVolume/Public/AudioGameplayVolumeComponent.h
+++ b/Engine/Plugins/AudioGameplayVolume/Source/AudioGameplayVolume/Public/AudioGameplayVolumeComponent.h
@@ -48,7 +48,7 @@ protected:
 
 	// A representation of this volume for the audio thread
 	UPROPERTY(Instanced, EditAnywhere, BlueprintReadWrite, Category = "AudioGameplay", Meta = (ShowOnlyInnerProperties, AllowPrivateAccess = "true"))
-	UAudioGameplayVolumeProxy* Proxy = nullptr;
+	TObjectPtr<UAudioGameplayVolumeProxy> Proxy = nullptr;
 
 #if WITH_EDITOR
 	virtual void PostEditChangeProperty(FPropertyChangedEvent& PropertyChangedEvent) override;
diff --git a/Engine/Plugins/AudioGameplayVolume/Source/AudioGameplayVolume/Public/AudioGameplayVolumeSubsystem.h b/Engine/Plugins/AudioGameplayVolume/Source/AudioGameplayVolume/Public/AudioGameplayVolumeSubsystem.h
--- a/Engine/Plugins/AudioGameplayVolume/Source/AudioGameplayVolume/Public/AudioGameplayVolumeSubsystem.h
+++ b/Engine/Plugins/AudioGameplayVolume/Source/AudioGameplayVolume/Public/AudioGameplayVolumeSubsystem.h
@@ -167,7 +167,7 @@ protected:
 
 	// Components in our system
 	UPROPERTY(Transient)
-	TMap<uint32, const UAudioGameplayVolumeComponent*> AGVComponents;
+	TMap<uint32, TObjectPtr<const UAudioGameplayVolumeComponent>> AGVComponents;
 
 	// Audio thread representation of Listeners
 	TArray<FAudioGameplayVolumeListener> AGVListeners;
diff --git a/Engine/Plugins/Compositing/Composure/Source/Composure/Classes/CompositingCaptureBase.h b/Engine/Plugins/Compositing/Composure/Source/Composure/Classes/CompositingCaptureBase.h
--- a/Engine/Plugins/Compositing/Composure/Source/Composure/Classes/CompositingCaptureBase.h
+++ b/Engine/Plugins/Compositing/Composure/Source/Composure/Classes/CompositingCaptureBase.h
@@ -24,7 +24,7 @@ class COMPOSURE_API ACompositingCaptureBase : public ACompositingElement
 public:
 	/** Component used to generate the scene capture for this CG Layer */
 	UPROPERTY(VisibleDefaultsOnly, BlueprintReadOnly, Category = "SceneCapture")
-	USceneCaptureComponent2D* SceneCaptureComponent2D = nullptr;
+	TObjectPtr<USceneCaptureComponent2D> SceneCaptureComponent2D = nullptr;
 
 protected:
 	/** Whether to apply distortion as a post-process effect on this CG Layer */
@@ -45,7 +45,7 @@ protected:
 
 	/** Cached distortion MID produced by the Lens Distortion Handler, used to clean up the post-process materials in the case that the the MID changes */
 	UPROPERTY()
-	UMaterialInstanceDynamic* LastDistortionMID = nullptr;
+	TObjectPtr<UMaterialInstanceDynamic> LastDistortionMID = nullptr;
 
 #if WITH_EDITORONLY_DATA
 PRAGMA_DISABLE_DEPRECATION_WARNINGS
diff --git a/Engine/Plugins/Compositing/Composure/Source/Composure/Classes/CompositingElement.h b/Engine/Plugins/Compositing/Composure/Source/Composure/Classes/CompositingElement.h
--- a/Engine/Plugins/Compositing/Composure/Source/Composure/Classes/CompositingElement.h
+++ b/Engine/Plugins/Compositing/Composure/Source/Composure/Classes/CompositingElement.h
@@ -79,10 +79,10 @@ class COMPOSURE_API ACompositingElement : public AComposurePipelineBaseActor, pu
 
 public:
 	UPROPERTY(BlueprintReadOnly, Category = "Composure")
-	UComposureCompositingTargetComponent* CompositingTarget;
+	TObjectPtr<UComposureCompositingTargetComponent> CompositingTarget;
 
 	UPROPERTY(BlueprintReadOnly, Category = "Composure")
-	UComposurePostProcessingPassProxy* PostProcessProxy;
+	TObjectPtr<UComposurePostProcessingPassProxy> PostProcessProxy;
 
 protected:
 	/*********************************/
@@ -90,13 +90,13 @@ protected:
 	//   - protected to prevent users from directly modifying these lists (use the accessor functions instead)
 
 	UPROPERTY(EditAnywhere, Instanced, BlueprintReadOnly, BlueprintGetter = GetInputsList, Category = "Composure|Input", meta=(ShowOnlyInnerProperties))
-	TArray<UCompositingElementInput*> Inputs;
+	TArray<TObjectPtr<UCompositingElementInput>> Inputs;
 
 	UPROPERTY(EditAnywhere, Instanced, BlueprintReadOnly, BlueprintGetter = GetTransformsList, Category = "Composure|Transform/Compositing Passes", meta=(DisplayName = "Transform Passes", ShowOnlyInnerProperties, DisplayAfter="Inputs"))
-	TArray<UCompositingElementTransform*> TransformPasses;
+	TArray<TObjectPtr<UCompositingElementTransform>> TransformPasses;
 
 	UPROPERTY(EditAnywhere, Instanced, BlueprintReadOnly, BlueprintGetter = GetOutputsList, Category = "Composure|Output", meta = (ShowOnlyInnerProperties))
-	TArray<UCompositingElementOutput*> Outputs;
+	TArray<TObjectPtr<UCompositingElementOutput>> Outputs;
 
 public:
 
@@ -147,7 +147,7 @@ public:
 	EInheritedSourceType PreviewTransformSource = EInheritedSourceType::Inherited;
 
 	UPROPERTY(EditAnywhere, Instanced, Category = "Composure|Preview")
-	UCompositingElementTransform* PreviewTransform;
+	TObjectPtr<UCompositingElementTransform> PreviewTransform;
 
 	UPROPERTY(EditDefaultsOnly, Category = "Composure|Editor")
 	TSubclassOf<UCompositingElementInput> DefaultInputType;
@@ -448,31 +448,31 @@ private:
 	FName CompShotIdName;
 
 	UPROPERTY()
-	ACompositingElement* Parent;
+	TObjectPtr<ACompositingElement> Parent;
 	UPROPERTY()
-	TArray<ACompositingElement*> ChildLayers;
+	TArray<TObjectPtr<ACompositingElement>> ChildLayers;
 
 	/** EDITOR ONLY - Properties associated with */
 #if WITH_EDITORONLY_DATA
 	UPROPERTY(Transient)
-	UTexture* DisabledMsgImage;
+	TObjectPtr<UTexture> DisabledMsgImage;
 	UPROPERTY(Transient)
-	UTexture* EmptyWarnImage;
+	TObjectPtr<UTexture> EmptyWarnImage;
 	UPROPERTY(Transient)
-	UTexture* SuspendedDbgImage;
+	TObjectPtr<UTexture> SuspendedDbgImage;
 	UPROPERTY(Transient)
-	UTexture* CompilerErrImage;
+	TObjectPtr<UTexture> CompilerErrImage;
 
 	UPROPERTY(Transient, DuplicateTransient)
 	bool bUsingDebugDisplayImage = false;
 
 	UPROPERTY(Transient, DuplicateTransient)
-	UTexture* ColorPickerDisplayImage;
+	TObjectPtr<UTexture> ColorPickerDisplayImage;
 	UPROPERTY(Transient, DuplicateTransient)
-	UTexture* EditorPreviewImage;
+	TObjectPtr<UTexture> EditorPreviewImage;
 
 	UPROPERTY(Transient, DuplicateTransient)
-	UTextureRenderTarget2D* ColorPickerTarget;
+	TObjectPtr<UTextureRenderTarget2D> ColorPickerTarget;
 
 	uint32 LastEnqueuedFrameId = (uint32)-1;
 	int32  PreviewCount = 0;
@@ -489,11 +489,11 @@ private:
 	 * re-construction, so we don't perpetually grow the lists.
 	 */
 	UPROPERTY()
-	TMap<UCompositingElementInput*, ECompPassConstructionType> UserConstructedInputs;
+	TMap<TObjectPtr<UCompositingElementInput>, ECompPassConstructionType> UserConstructedInputs;
 	UPROPERTY()
-	TMap<UCompositingElementTransform*, ECompPassConstructionType> UserConstructedTransforms;
+	TMap<TObjectPtr<UCompositingElementTransform>, ECompPassConstructionType> UserConstructedTransforms;
 	UPROPERTY()
-	TMap<UCompositingElementOutput*, ECompPassConstructionType> UserConstructedOutputs;
+	TMap<TObjectPtr<UCompositingElementOutput>, ECompPassConstructionType> UserConstructedOutputs;
 
 	/** 
 	 * Authoritative lists that we use to iterate on the passes - conjoined from the public lists and the  
@@ -502,14 +502,14 @@ private:
 	 * alive via the transaction buffer).
 	 */
 	UPROPERTY(Instanced, Transient, DuplicateTransient, SkipSerialization)
-	TArray<UCompositingElementInput*> InternalInputs;
+	TArray<TObjectPtr<UCompositingElementInput>> InternalInputs;
 	UPROPERTY(Instanced, Transient, DuplicateTransient, SkipSerialization)
-	TArray<UCompositingElementTransform*> InternalTransformPasses;
+	TArray<TObjectPtr<UCompositingElementTransform>> InternalTransformPasses;
 	UPROPERTY(Instanced, Transient, DuplicateTransient, SkipSerialization)
-	TArray<UCompositingElementOutput*> InternalOutputs;
+	TArray<TObjectPtr<UCompositingElementOutput>> InternalOutputs;
 
 	UPROPERTY(Transient)
-	UAlphaTransformPass* InternalAlphaPass = nullptr;
+	TObjectPtr<UAlphaTransformPass> InternalAlphaPass = nullptr;
 
 	/** */
 	FCompositingTextureLookupTable PassResultsTable;
diff --git a/Engine/Plugins/Compositing/Composure/Source/Composure/Classes/CompositingElements/CompositingElementInputs.h b/Engine/Plugins/Compositing/Composure/Source/Composure/Classes/CompositingElements/CompositingElementInputs.h
--- a/Engine/Plugins/Compositing/Composure/Source/Composure/Classes/CompositingElements/CompositingElementInputs.h
+++ b/Engine/Plugins/Compositing/Composure/Source/Composure/Classes/CompositingElements/CompositingElementInputs.h
@@ -40,11 +40,11 @@ private:
 	virtual UTexture* GetMediaTexture() const PURE_VIRTUAL(UCompositingMediaInput::GetMediaTexture, return nullptr;);
 
 	UPROPERTY(Transient)
-	UMaterialInterface* DefaultMaterial;
+	TObjectPtr<UMaterialInterface> DefaultMaterial;
 	UPROPERTY(Transient)
-	UMaterialInterface* DefaultTestPlateMaterial;
+	TObjectPtr<UMaterialInterface> DefaultTestPlateMaterial;
 	UPROPERTY(Transient, DuplicateTransient, SkipSerialization)
-	UMaterialInstanceDynamic* FallbackMID;
+	TObjectPtr<UMaterialInstanceDynamic> FallbackMID;
 };
 
 /* UMediaBundleCompositingInput
@@ -80,7 +80,7 @@ public:
 	UMediaTextureCompositingInput();
 
 	UPROPERTY(EditAnywhere, BlueprintReadOnly, Category = "Compositing Pass", meta = (DisplayAfter = "PassName", EditCondition = "bEnabled"))
-	UMediaTexture* MediaSource;
+	TObjectPtr<UMediaTexture> MediaSource;
 
 private:
 	//~ UCompositingMediaInput interface	
diff --git a/Engine/Plugins/Compositing/Composure/Source/Composure/Classes/CompositingElements/CompositingElementOutputs.h b/Engine/Plugins/Compositing/Composure/Source/Composure/Classes/CompositingElements/CompositingElementOutputs.h
--- a/Engine/Plugins/Compositing/Composure/Source/Composure/Classes/CompositingElements/CompositingElementOutputs.h
+++ b/Engine/Plugins/Compositing/Composure/Source/Composure/Classes/CompositingElements/CompositingElementOutputs.h
@@ -23,7 +23,7 @@ class COMPOSURE_API UColorConverterOutputPass : public UCompositingElementOutput
 
 public:
 	UPROPERTY(EditAnywhere, Instanced, Category="Compositing Pass", meta=(DisplayName="Color Conversion", DisplayAfter = "PassName", EditCondition = "bEnabled", ShowOnlyInnerProperties))
-	UCompositingElementTransform* ColorConverter;
+	TObjectPtr<UCompositingElementTransform> ColorConverter;
 
 public:
 	//~ Begin UObject interface
@@ -53,7 +53,7 @@ protected:
 
 #if WITH_EDITORONLY_DATA
 	UPROPERTY(Transient)
-	UTexture* PreviewResult = nullptr;
+	TObjectPtr<UTexture> PreviewResult = nullptr;
 #endif
 
 private:
@@ -77,7 +77,7 @@ public:
 	UCompositingMediaCaptureOutput();
 
 	UPROPERTY(EditAnywhere, Category="Compositing Pass", BlueprintReadWrite, meta = (DisplayAfter = "PassName", EditCondition = "bEnabled"))
-	UMediaOutput* CaptureOutput;
+	TObjectPtr<UMediaOutput> CaptureOutput;
 
 	/** */
 	bool IsCapturing() const;
@@ -102,7 +102,7 @@ protected:
 
 private:
 	UPROPERTY(Transient, DuplicateTransient, SkipSerialization)
-	UMediaCapture* ActiveCapture = nullptr;
+	TObjectPtr<UMediaCapture> ActiveCapture = nullptr;
 };
 
 /* URenderTargetCompositingOutput
@@ -117,7 +117,7 @@ class COMPOSURE_API URenderTargetCompositingOutput : public UCompositingElementO
 
 public:
 	UPROPERTY(EditAnywhere, BlueprintReadWrite,Category="Compositing Pass", meta = (DisplayAfter = "PassName", EditCondition = "bEnabled"))
-	UTextureRenderTarget2D* RenderTarget;
+	TObjectPtr<UTextureRenderTarget2D> RenderTarget;
 
 public:
 	//~ UCompositingElementOutput interface
diff --git a/Engine/Plugins/Compositing/Composure/Source/Composure/Classes/CompositingElements/CompositingElementTransforms.h b/Engine/Plugins/Compositing/Composure/Source/Composure/Classes/CompositingElements/CompositingElementTransforms.h
--- a/Engine/Plugins/Compositing/Composure/Source/Composure/Classes/CompositingElements/CompositingElementTransforms.h
+++ b/Engine/Plugins/Compositing/Composure/Source/Composure/Classes/CompositingElements/CompositingElementTransforms.h
@@ -25,7 +25,7 @@ public:
 	float RenderScale = 1.f;
 
 	UPROPERTY(EditAnywhere, Instanced, BlueprintReadOnly, Category = "Compositing Pass", meta = (DisplayAfter = "RenderScale"))
-	TArray<UComposurePostProcessPassPolicy*> PostProcessPasses;
+	TArray<TObjectPtr<UComposurePostProcessPassPolicy>> PostProcessPasses;
 
 public:
 	//~ Begin UCompositingElementTransform interface
@@ -103,7 +103,7 @@ public:
 
 private:
 	UPROPERTY(Transient, DuplicateTransient, SkipSerialization)
-	UComposureTonemapperPassPolicy* TonemapPolicy;
+	TObjectPtr<UComposureTonemapperPassPolicy> TonemapPolicy;
 };
 
 /* UMultiPassChromaKeyer
@@ -133,7 +133,7 @@ public:
 
 private:
 	UPROPERTY(Transient)
-	UTexture* DefaultWhiteTexture = nullptr;
+	TObjectPtr<UTexture> DefaultWhiteTexture = nullptr;
 };
 
 
@@ -164,7 +164,7 @@ public:
 
 private:
 	UPROPERTY(Transient)
-	UTexture* DefaultWhiteTexture = nullptr;
+	TObjectPtr<UTexture> DefaultWhiteTexture = nullptr;
 };
 
 /* UAlphaTransformPass
@@ -190,9 +190,9 @@ public:
 
 private:
 	UPROPERTY(Transient)
-	UMaterialInterface* DefaultMaterial;
+	TObjectPtr<UMaterialInterface> DefaultMaterial;
 	UPROPERTY(Transient)
-	UMaterialInstanceDynamic* AlphaTransformMID;
+	TObjectPtr<UMaterialInstanceDynamic> AlphaTransformMID;
 };
 
 
diff --git a/Engine/Plugins/Compositing/Composure/Source/Composure/Classes/CompositingElements/PlayerViewportCompositingOutput.h b/Engine/Plugins/Compositing/Composure/Source/Composure/Classes/CompositingElements/PlayerViewportCompositingOutput.h
--- a/Engine/Plugins/Compositing/Composure/Source/Composure/Classes/CompositingElements/PlayerViewportCompositingOutput.h
+++ b/Engine/Plugins/Compositing/Composure/Source/Composure/Classes/CompositingElements/PlayerViewportCompositingOutput.h
@@ -60,15 +60,15 @@ private:
 	int32 ActiveOverrideIndex = INDEX_NONE;
 
 	UPROPERTY(Transient, DuplicateTransient, SkipSerialization)
-	UPlayerCompOutputCameraModifier* ActiveCamModifier;
+	TObjectPtr<UPlayerCompOutputCameraModifier> ActiveCamModifier;
 
 	UPROPERTY(Transient)
-	UMaterialInterface* TonemapperBaseMat;
+	TObjectPtr<UMaterialInterface> TonemapperBaseMat;
 	UPROPERTY(Transient)
-	UMaterialInterface* PreTonemapBaseMat;
+	TObjectPtr<UMaterialInterface> PreTonemapBaseMat;
 
 	UPROPERTY(Transient, DuplicateTransient, SkipSerialization)
-	UMaterialInstanceDynamic* ViewportOverrideMID;
+	TObjectPtr<UMaterialInstanceDynamic> ViewportOverrideMID;
 
 	TWeakObjectPtr<APlayerController> TargetedPlayerController;
 };
@@ -90,6 +90,6 @@ public:
 
 private:
 	UPROPERTY(Transient)
-	UPlayerViewportCompositingOutput* Owner;
+	TObjectPtr<UPlayerViewportCompositingOutput> Owner;
 };
 
diff --git a/Engine/Plugins/Compositing/Composure/Source/Composure/Classes/ComposureLensBloomPass.h b/Engine/Plugins/Compositing/Composure/Source/Composure/Classes/ComposureLensBloomPass.h
--- a/Engine/Plugins/Compositing/Composure/Source/Composure/Classes/ComposureLensBloomPass.h
+++ b/Engine/Plugins/Compositing/Composure/Source/Composure/Classes/ComposureLensBloomPass.h
@@ -39,7 +39,7 @@ public:
 
 private:
 	UPROPERTY(Transient)
-	UMaterialInstanceDynamic* TonemapperReplacingMID;
+	TObjectPtr<UMaterialInstanceDynamic> TonemapperReplacingMID;
 };
 
 /**
@@ -57,7 +57,7 @@ public:
 	FLensBloomSettings Settings;
 
 	UPROPERTY(EditAnywhere, Category = "Tonemapper Settings", AdvancedDisplay)
-	UMaterialInterface* ReplacementMaterial;
+	TObjectPtr<UMaterialInterface> ReplacementMaterial;
 
 	UPROPERTY(EditAnywhere, Category = "Tonemapper Settings", AdvancedDisplay)
 	FName BloomIntensityParamName;
@@ -68,5 +68,5 @@ public:
 
 private:
 	UPROPERTY(Transient)
-	UMaterialInstanceDynamic* TonemapperReplacmentMID;
+	TObjectPtr<UMaterialInstanceDynamic> TonemapperReplacmentMID;
 };
diff --git a/Engine/Plugins/Compositing/Composure/Source/Composure/Classes/ComposurePlayerCompositingTarget.h b/Engine/Plugins/Compositing/Composure/Source/Composure/Classes/ComposurePlayerCompositingTarget.h
--- a/Engine/Plugins/Compositing/Composure/Source/Composure/Classes/ComposurePlayerCompositingTarget.h
+++ b/Engine/Plugins/Compositing/Composure/Source/Composure/Classes/ComposurePlayerCompositingTarget.h
@@ -65,15 +65,15 @@ public:
 private:
 	// Current player camera manager the target is bind on.
 	UPROPERTY(Transient)
-	APlayerCameraManager* PlayerCameraManager;
+	TObjectPtr<APlayerCameraManager> PlayerCameraManager;
 
 	// Underlying player camera modifier
 	UPROPERTY(Transient)
-	class UComposurePlayerCompositingCameraModifier* PlayerCameraModifier;
+	TObjectPtr<class UComposurePlayerCompositingCameraModifier> PlayerCameraModifier;
 	
 	// Post process material that replaces the tonemapper to dump the player's render target.
 	UPROPERTY(Transient)
-	class UMaterialInstanceDynamic* ReplaceTonemapperMID;
+	TObjectPtr<class UMaterialInstanceDynamic> ReplaceTonemapperMID;
 
 	// Backup of the engine showflags to restore when unbinding the compositing target from the camera manager.
 	FEngineShowFlags EngineShowFlagsBackup;
@@ -132,11 +132,11 @@ public:
 
 private:
 	UPROPERTY(Transient)
-	UTexture* DisplayTexture;
+	TObjectPtr<UTexture> DisplayTexture;
 
 #if WITH_EDITORONLY_DATA
 	UPROPERTY(Transient)
-	UTexture* CompilerErrImage;
+	TObjectPtr<UTexture> CompilerErrImage;
 
 	int32 PreviewCount = 0;
 	bool bUseImplicitGammaOnPreview = true;
diff --git a/Engine/Plugins/Compositing/Composure/Source/Composure/Classes/ComposurePostProcessPass.h b/Engine/Plugins/Compositing/Composure/Source/Composure/Classes/ComposurePostProcessPass.h
--- a/Engine/Plugins/Compositing/Composure/Source/Composure/Classes/ComposurePostProcessPass.h
+++ b/Engine/Plugins/Compositing/Composure/Source/Composure/Classes/ComposurePostProcessPass.h
@@ -64,19 +64,19 @@ public:
 protected:
 	// Underlying scene capture.
 	UPROPERTY(Transient, NonTransactional)
-	USceneCaptureComponent2D* SceneCapture;
+	TObjectPtr<USceneCaptureComponent2D> SceneCapture;
 
 	// Blendable interface to intercept the OverrideBlendableSettings.
 	UPROPERTY(Transient, NonTransactional)
-	UComposurePostProcessBlendable* BlendableInterface;
+	TObjectPtr<UComposurePostProcessBlendable> BlendableInterface;
 
 	// Setup post process material.
 	UPROPERTY(Transient, NonTransactional)
-	UMaterialInterface* SetupMaterial;
+	TObjectPtr<UMaterialInterface> SetupMaterial;
 
 	// Internal material that replace the tonemapper to output linear color space.
 	UPROPERTY(Transient, NonTransactional)
-	UMaterialInterface* TonemapperReplacement;
+	TObjectPtr<UMaterialInterface> TonemapperReplacement;
 
 	// Called by UComposurePostProcessBlendable::OverrideBlendableSettings.
 	void OverrideBlendableSettings(class FSceneView& View, float Weight) const;
diff --git a/Engine/Plugins/Compositing/Composure/Source/Composure/Classes/ComposurePostProcessingPassProxy.h b/Engine/Plugins/Compositing/Composure/Source/Composure/Classes/ComposurePostProcessingPassProxy.h
--- a/Engine/Plugins/Compositing/Composure/Source/Composure/Classes/ComposurePostProcessingPassProxy.h
+++ b/Engine/Plugins/Compositing/Composure/Source/Composure/Classes/ComposurePostProcessingPassProxy.h
@@ -48,5 +48,5 @@ public:
 
 private:
 	UPROPERTY(Transient, DuplicateTransient, SkipSerialization)
-	UMaterialInstanceDynamic* SetupMID;
+	TObjectPtr<UMaterialInstanceDynamic> SetupMID;
 };
diff --git a/Engine/Plugins/Compositing/Composure/Source/Composure/Private/ComposureConfigSettings.h b/Engine/Plugins/Compositing/Composure/Source/Composure/Private/ComposureConfigSettings.h
--- a/Engine/Plugins/Compositing/Composure/Source/Composure/Private/ComposureConfigSettings.h
+++ b/Engine/Plugins/Compositing/Composure/Source/Composure/Private/ComposureConfigSettings.h
@@ -26,5 +26,5 @@ public:
 	
 private:
 	UPROPERTY(Transient)
-	UTexture* FallbackCompositingTextureObj;
+	TObjectPtr<UTexture> FallbackCompositingTextureObj;
 };
diff --git a/Engine/Plugins/Compositing/Composure/Source/Composure/Private/ComposurePostProcessBlendable.h b/Engine/Plugins/Compositing/Composure/Source/Composure/Private/ComposurePostProcessBlendable.h
--- a/Engine/Plugins/Compositing/Composure/Source/Composure/Private/ComposurePostProcessBlendable.h
+++ b/Engine/Plugins/Compositing/Composure/Source/Composure/Private/ComposurePostProcessBlendable.h
@@ -28,7 +28,7 @@ public:
 private:
 	// Current player camera manager the target is bind on.
 	UPROPERTY(Transient)
-	class UComposurePostProcessPass* Target;
+	TObjectPtr<class UComposurePostProcessPass> Target;
 	
 	friend class UComposurePostProcessPass;
 };
diff --git a/Engine/Plugins/Compositing/Composure/Source/Composure/Private/EditorSupport/CompositingEditorSupportLibrary.h b/Engine/Plugins/Compositing/Composure/Source/Composure/Private/EditorSupport/CompositingEditorSupportLibrary.h
--- a/Engine/Plugins/Compositing/Composure/Source/Composure/Private/EditorSupport/CompositingEditorSupportLibrary.h
+++ b/Engine/Plugins/Compositing/Composure/Source/Composure/Private/EditorSupport/CompositingEditorSupportLibrary.h
@@ -50,9 +50,9 @@ private:
 	void InternalOnCancel();
 
 	UPROPERTY(Transient)
-	UTextureRenderTarget2D* PickerTarget;
+	TObjectPtr<UTextureRenderTarget2D> PickerTarget;
 	UPROPERTY(Transient)
-	UTexture* PickerDisplayImage;
+	TObjectPtr<UTexture> PickerDisplayImage;
 
 	bool bUseImplicitGamma = true;
 };
diff --git a/Engine/Plugins/Compositing/Composure/Source/Composure/Public/CompositingElements/CompositingMaterialPass.h b/Engine/Plugins/Compositing/Composure/Source/Composure/Public/CompositingElements/CompositingMaterialPass.h
--- a/Engine/Plugins/Compositing/Composure/Source/Composure/Public/CompositingElements/CompositingMaterialPass.h
+++ b/Engine/Plugins/Compositing/Composure/Source/Composure/Public/CompositingElements/CompositingMaterialPass.h
@@ -98,7 +98,7 @@ public:
 	FCompositingMaterial();
 
 	UPROPERTY(EditAnywhere, BlueprintReadWrite, Category = "CompositingMaterial")
-	UMaterialInterface* Material;
+	TObjectPtr<UMaterialInterface> Material;
 
 	/** Maps material texture param names to prior passes/elements. Overrides the element's param mapping list above. */
 	UPROPERTY(EditAnywhere, BlueprintReadWrite, Category = "CompositingMaterial", meta = (ReadOnlyKeys, DisplayName = "Input Elements"))
@@ -164,6 +164,6 @@ private:
 	bool bParamsModified = true;
 
 	UPROPERTY(Transient, SkipSerialization)
-	UMaterialInstanceDynamic* CachedMID;
+	TObjectPtr<UMaterialInstanceDynamic> CachedMID;
 };
 
diff --git a/Engine/Plugins/Compositing/Composure/Source/Composure/Public/ComposureUVMap.h b/Engine/Plugins/Compositing/Composure/Source/Composure/Public/ComposureUVMap.h
--- a/Engine/Plugins/Compositing/Composure/Source/Composure/Public/ComposureUVMap.h
+++ b/Engine/Plugins/Compositing/Composure/Source/Composure/Public/ComposureUVMap.h
@@ -33,7 +33,7 @@ struct COMPOSURE_API FComposureUVMapSettings
 
 	/** Displacement texture to use. */
 	UPROPERTY(EditAnywhere, BlueprintReadWrite, Category = "UV Mapping")
-	class UTexture* DisplacementTexture;
+	TObjectPtr<class UTexture> DisplacementTexture;
 
 	/** Whether to use blue and alpha channel instead of red and green channel in computation of DeltaUV. */
 	UPROPERTY(EditAnywhere, BlueprintReadWrite, Category = "UV Mapping")
diff --git a/Engine/Plugins/Compositing/OpenColorIO/Source/OpenColorIO/Public/OpenColorIOColorTransform.h b/Engine/Plugins/Compositing/OpenColorIO/Source/OpenColorIO/Public/OpenColorIOColorTransform.h
--- a/Engine/Plugins/Compositing/OpenColorIO/Source/OpenColorIO/Public/OpenColorIOColorTransform.h
+++ b/Engine/Plugins/Compositing/OpenColorIO/Source/OpenColorIO/Public/OpenColorIOColorTransform.h
@@ -146,7 +146,7 @@ public:
 public:
 
 	UPROPERTY()
-	UOpenColorIOConfiguration* ConfigurationOwner;
+	TObjectPtr<UOpenColorIOConfiguration> ConfigurationOwner;
 
 	UPROPERTY(VisibleAnywhere, Category = "ColorSpace")
 	FString SourceColorSpace;
@@ -159,7 +159,7 @@ private:
 	/** If the color space requires a 3dLUT, this will contains the data to do the transform */
 	/** Note: This will be serialized when cooking. Otherwhise, it relies on raw data of the library and what's on DDC */
 	UPROPERTY(Transient)
-	UVolumeTexture* Lut3dTexture = nullptr;
+	TObjectPtr<UVolumeTexture> Lut3dTexture = nullptr;
 	
 	/** Inline ColorTransform resources serialized from disk. To be processed on game thread in PostLoad. */
 	TArray<FOpenColorIOTransformResource> LoadedTransformResources;
diff --git a/Engine/Plugins/Compositing/OpenColorIO/Source/OpenColorIO/Public/OpenColorIOConfiguration.h b/Engine/Plugins/Compositing/OpenColorIO/Source/OpenColorIO/Public/OpenColorIOConfiguration.h
--- a/Engine/Plugins/Compositing/OpenColorIO/Source/OpenColorIO/Public/OpenColorIOConfiguration.h
+++ b/Engine/Plugins/Compositing/OpenColorIO/Source/OpenColorIO/Public/OpenColorIOConfiguration.h
@@ -92,7 +92,7 @@ public:
 private:
 
 	UPROPERTY()
-	TArray<UOpenColorIOColorTransform*> ColorTransforms;
+	TArray<TObjectPtr<UOpenColorIOColorTransform>> ColorTransforms;
 
 #if WITH_EDITORONLY_DATA && WITH_OCIO
 	OCIO_NAMESPACE::ConstConfigRcPtr LoadedConfig;
diff --git a/Engine/Plugins/Editor/ModelingToolsEditorMode/Source/ModelingToolsEditorMode/Public/ModelingSelectionInteraction.h b/Engine/Plugins/Editor/ModelingToolsEditorMode/Source/ModelingToolsEditorMode/Public/ModelingSelectionInteraction.h
--- a/Engine/Plugins/Editor/ModelingToolsEditorMode/Source/ModelingToolsEditorMode/Public/ModelingSelectionInteraction.h
+++ b/Engine/Plugins/Editor/ModelingToolsEditorMode/Source/ModelingToolsEditorMode/Public/ModelingSelectionInteraction.h
@@ -60,11 +60,11 @@ public:
 public:
 	// click-to-select behavior
 	UPROPERTY()
-	USingleClickInputBehavior* ClickBehavior;
+	TObjectPtr<USingleClickInputBehavior> ClickBehavior;
 
 	// set of all behaviors, will be passed up to UInputRouter
 	UPROPERTY()
-	UInputBehaviorSet* BehaviorSet;
+	TObjectPtr<UInputBehaviorSet> BehaviorSet;
 
 public:
 	//
diff --git a/Engine/Plugins/Enterprise/DataprepEditor/Source/DataprepCore/Private/Parameterization/DataprepParameterization.h b/Engine/Plugins/Enterprise/DataprepEditor/Source/DataprepCore/Private/Parameterization/DataprepParameterization.h
--- a/Engine/Plugins/Enterprise/DataprepEditor/Source/DataprepCore/Private/Parameterization/DataprepParameterization.h
+++ b/Engine/Plugins/Enterprise/DataprepEditor/Source/DataprepCore/Private/Parameterization/DataprepParameterization.h
@@ -86,7 +86,7 @@ struct FDataprepParameterizationBinding
 	bool operator==(const FDataprepParameterizationBinding& Other) const;
 
 	UPROPERTY()
-	UDataprepParameterizableObject* ObjectBinded = nullptr;
+	TObjectPtr<UDataprepParameterizableObject> ObjectBinded = nullptr;
 
 	UPROPERTY()
 	TArray<FDataprepPropertyLink> PropertyChain;
@@ -346,7 +346,7 @@ public:
 private:
 	// The containers for the bindings
 	UPROPERTY()
-	UDataprepParameterizationBindings* BindingsContainer;
+	TObjectPtr<UDataprepParameterizationBindings> BindingsContainer;
 
 	/**
 	* Mapping of the observed object to their delegate handle
@@ -359,10 +359,10 @@ private:
 	TMap<FName, FProperty*> NameToParameterizationProperty;
 
 	UPROPERTY(Transient, NonTransactional)
-	UClass* CustomContainerClass;
+	TObjectPtr<UClass> CustomContainerClass;
 
 	UPROPERTY(Transient, NonTransactional)
-	UDataprepParameterizableObject* DefaultParameterisation;
+	TObjectPtr<UDataprepParameterizableObject> DefaultParameterisation;
 
 	/** 
 	 * This is used only to store a serialization of the values of the parameterization since we can't save our custom container class
@@ -440,11 +440,11 @@ private:
 
 	// The parameterization from which this instance was constructed
 	UPROPERTY()
-	UDataprepParameterization* SourceParameterization;
+	TObjectPtr<UDataprepParameterization> SourceParameterization;
 
 	// The actual object on which the parameterization data is stored
 	UPROPERTY(Transient, NonTransactional)
-	UObject* ParameterizationInstance;
+	TObjectPtr<UObject> ParameterizationInstance;
 
 	// This is used only to store a serialization of the values of the parameterization since we can't save the custom class
 	UPROPERTY()
diff --git a/Engine/Plugins/Enterprise/DataprepEditor/Source/DataprepCore/Public/Blueprint/K2Node_DataprepActionCore.h b/Engine/Plugins/Enterprise/DataprepEditor/Source/DataprepCore/Public/Blueprint/K2Node_DataprepActionCore.h
--- a/Engine/Plugins/Enterprise/DataprepEditor/Source/DataprepCore/Public/Blueprint/K2Node_DataprepActionCore.h
+++ b/Engine/Plugins/Enterprise/DataprepEditor/Source/DataprepCore/Public/Blueprint/K2Node_DataprepActionCore.h
@@ -25,7 +25,7 @@ public:
 
 protected:
 	UPROPERTY()
-	UDataprepActionAsset* DataprepActionAsset;
+	TObjectPtr<UDataprepActionAsset> DataprepActionAsset;
 };
 
 
diff --git a/Engine/Plugins/Enterprise/DataprepEditor/Source/DataprepCore/Public/DataprepActionAsset.h b/Engine/Plugins/Enterprise/DataprepEditor/Source/DataprepCore/Public/DataprepActionAsset.h
--- a/Engine/Plugins/Enterprise/DataprepEditor/Source/DataprepCore/Public/DataprepActionAsset.h
+++ b/Engine/Plugins/Enterprise/DataprepEditor/Source/DataprepCore/Public/DataprepActionAsset.h
@@ -87,7 +87,7 @@ private:
 
 	// The actual object of the step
 	UPROPERTY()
-	UDataprepParameterizableObject* StepObject;
+	TObjectPtr<UDataprepParameterizableObject> StepObject;
 
 	// Will be used for future error message if the step object can't be loaded
 	UPROPERTY()
@@ -95,12 +95,12 @@ private:
 
 	// The operation will only be not null if the step is a operation
 	UPROPERTY(meta = (DeprecatedProperty, DeprecationMessage = "Use the step and step type instead."))
-	UDataprepOperation* Operation_DEPRECATED;
+	TObjectPtr<UDataprepOperation> Operation_DEPRECATED;
 
 	
 	// The Filter will only be not null if the step is a Filter/Selector
 	UPROPERTY(meta = (DeprecatedProperty, DeprecationMessage = "Use the step and step type instead."))
-	UDataprepFilter* Filter_DEPRECATED;
+	TObjectPtr<UDataprepFilter> Filter_DEPRECATED;
 };
 
 /** Structure to pass execution context to action */
@@ -463,10 +463,10 @@ private:
 
 	/** Array of operations and/or filters constituting this action */
 	UPROPERTY()
-	TArray<UDataprepActionStep*> Steps;
+	TArray<TObjectPtr<UDataprepActionStep>> Steps;
 
 	UPROPERTY()
-	UDataprepActionAppearance* Appearance;
+	TObjectPtr<UDataprepActionAppearance> Appearance;
 
 	/** Broadcasts any change to the stack of steps */
 	FOnStepsOrderChanged OnStepsOrderChanged;
diff --git a/Engine/Plugins/Enterprise/DataprepEditor/Source/DataprepCore/Public/DataprepAsset.h b/Engine/Plugins/Enterprise/DataprepEditor/Source/DataprepCore/Public/DataprepAsset.h
--- a/Engine/Plugins/Enterprise/DataprepEditor/Source/DataprepCore/Public/DataprepAsset.h
+++ b/Engine/Plugins/Enterprise/DataprepEditor/Source/DataprepCore/Public/DataprepAsset.h
@@ -217,7 +217,7 @@ protected:
 #if WITH_EDITORONLY_DATA
 	/** DEPRECATED: Pointer to data preparation pipeline blueprint previously used to process input data */
 	UPROPERTY(meta = (DeprecatedProperty, DeprecationMessage = "Using directly ActionAssets property instead of Blueprint to manage actions."))
-	UBlueprint* DataprepRecipeBP_DEPRECATED;
+	TObjectPtr<UBlueprint> DataprepRecipeBP_DEPRECATED;
 	// end of temp code for nodes development
 
 	/** DEPRECATED: List of producers referenced by the asset */
@@ -226,7 +226,7 @@ protected:
 
 	/** DEPRECATED: COnsumer referenced by the asset */
 	UPROPERTY(meta = (DeprecatedProperty, DeprecationMessage = "Property moved to UDataprepAssetInterface as Output."))
-	UDataprepContentConsumer* Consumer_DEPRECATED;
+	TObjectPtr<UDataprepContentConsumer> Consumer_DEPRECATED;
 #endif
 
 private:
@@ -239,13 +239,13 @@ private:
 private:
 	/** DEPRECATED: Pointer to the entry node of the pipeline blueprint previously used to process input data */
 	UPROPERTY(meta = (DeprecatedProperty, DeprecationMessage = "Using directly ActionAssets property instead of Blueprint to manage actions."))
-	UEdGraphNode* StartNode_DEPRECATED;
+	TObjectPtr<UEdGraphNode> StartNode_DEPRECATED;
 
 	UPROPERTY()
-	UDataprepParameterization* Parameterization;
+	TObjectPtr<UDataprepParameterization> Parameterization;
 
 	UPROPERTY()
-	TArray<UDataprepActionAsset*> ActionAssets;
+	TArray<TObjectPtr<UDataprepActionAsset>> ActionAssets;
 
 	FOnDataprepActionAssetChange OnActionChanged;
 
diff --git a/Engine/Plugins/Enterprise/DataprepEditor/Source/DataprepCore/Public/DataprepAssetInstance.h b/Engine/Plugins/Enterprise/DataprepEditor/Source/DataprepCore/Public/DataprepAssetInstance.h
--- a/Engine/Plugins/Enterprise/DataprepEditor/Source/DataprepCore/Public/DataprepAssetInstance.h
+++ b/Engine/Plugins/Enterprise/DataprepEditor/Source/DataprepCore/Public/DataprepAssetInstance.h
@@ -58,10 +58,10 @@ public:
 protected:
 	/** Parent Dataprep asset's interface */
 	UPROPERTY()
-	UDataprepAssetInterface* Parent;
+	TObjectPtr<UDataprepAssetInterface> Parent;
 
 	UPROPERTY()
-	UDataprepParameterizationInstance* Parameterization;
+	TObjectPtr<UDataprepParameterizationInstance> Parameterization;
 
 	/** Delegate broadcasted when the consumer or one of the producers has changed */
 	FOnDataprepAssetInstanceChanged OnParentChanged;
@@ -72,5 +72,5 @@ private:
 
 private:
 	UPROPERTY()
-	TArray<UDataprepActionAsset*> ActionsFromDataprepAsset;
+	TArray<TObjectPtr<UDataprepActionAsset>> ActionsFromDataprepAsset;
 };
\ No newline at end of file
diff --git a/Engine/Plugins/Enterprise/DataprepEditor/Source/DataprepCore/Public/DataprepAssetInterface.h b/Engine/Plugins/Enterprise/DataprepEditor/Source/DataprepCore/Public/DataprepAssetInterface.h
--- a/Engine/Plugins/Enterprise/DataprepEditor/Source/DataprepCore/Public/DataprepAssetInterface.h
+++ b/Engine/Plugins/Enterprise/DataprepEditor/Source/DataprepCore/Public/DataprepAssetInterface.h
@@ -139,15 +139,15 @@ protected:
 
 	/** Recipe associated to the Dataprep asset */
 	UPROPERTY()
-	UDataprepRecipeInterface* Recipe;
+	TObjectPtr<UDataprepRecipeInterface> Recipe;
 
 	/** Producers associated to the Dataprep asset */
 	UPROPERTY()
-	UDataprepAssetProducers* Inputs;
+	TObjectPtr<UDataprepAssetProducers> Inputs;
 
 	/** Consumer associated to the Dataprep asset */
 	UPROPERTY()
-	UDataprepContentConsumer* Output;
+	TObjectPtr<UDataprepContentConsumer> Output;
 #endif
 
 protected:
diff --git a/Engine/Plugins/Enterprise/DataprepEditor/Source/DataprepCore/Public/DataprepAssetProducers.h b/Engine/Plugins/Enterprise/DataprepEditor/Source/DataprepCore/Public/DataprepAssetProducers.h
--- a/Engine/Plugins/Enterprise/DataprepEditor/Source/DataprepCore/Public/DataprepAssetProducers.h
+++ b/Engine/Plugins/Enterprise/DataprepEditor/Source/DataprepCore/Public/DataprepAssetProducers.h
@@ -32,7 +32,7 @@ struct FDataprepAssetProducer
 	{}
 
 	UPROPERTY()
-	UDataprepContentProducer* Producer = nullptr;
+	TObjectPtr<UDataprepContentProducer> Producer = nullptr;
 
 	UPROPERTY()
 	bool bIsEnabled = true;
diff --git a/Engine/Plugins/Enterprise/DataprepEditor/Source/DataprepCore/Public/DataprepFactories.h b/Engine/Plugins/Enterprise/DataprepEditor/Source/DataprepCore/Public/DataprepFactories.h
--- a/Engine/Plugins/Enterprise/DataprepEditor/Source/DataprepCore/Public/DataprepFactories.h
+++ b/Engine/Plugins/Enterprise/DataprepEditor/Source/DataprepCore/Public/DataprepFactories.h
@@ -38,7 +38,7 @@ public:
 	 * The parent of the of the instance to create
 	 */
 	UPROPERTY(EditAnywhere, Category="Settings")
-	class UDataprepAsset* Parent;
+	TObjectPtr<class UDataprepAsset> Parent;
 
 	//~ Begin UFactory Interface
 	virtual UObject* FactoryCreateNew(UClass* InClass, UObject* InParent, FName InName, EObjectFlags Flags, UObject* Context, FFeedbackContext* Warn) override;
diff --git a/Engine/Plugins/Enterprise/DataprepEditor/Source/DataprepCore/Public/DataprepOperation.h b/Engine/Plugins/Enterprise/DataprepEditor/Source/DataprepCore/Public/DataprepOperation.h
--- a/Engine/Plugins/Enterprise/DataprepEditor/Source/DataprepCore/Public/DataprepOperation.h
+++ b/Engine/Plugins/Enterprise/DataprepEditor/Source/DataprepCore/Public/DataprepOperation.h
@@ -76,7 +76,7 @@ struct FDataprepContext
 	 * This is the objects on which an operation can operate
 	 */
 	UPROPERTY( BlueprintReadOnly, Category = "Dataprep")
-	TArray<UObject*> Objects;
+	TArray<TObjectPtr<UObject>> Objects;
 };
 
 /**
diff --git a/Engine/Plugins/Enterprise/DataprepEditor/Source/DataprepCore/Public/SelectionSystem/DataprepBoolFilter.h b/Engine/Plugins/Enterprise/DataprepEditor/Source/DataprepCore/Public/SelectionSystem/DataprepBoolFilter.h
--- a/Engine/Plugins/Enterprise/DataprepEditor/Source/DataprepCore/Public/SelectionSystem/DataprepBoolFilter.h
+++ b/Engine/Plugins/Enterprise/DataprepEditor/Source/DataprepCore/Public/SelectionSystem/DataprepBoolFilter.h
@@ -34,5 +34,5 @@ private:
 	//~ Begin UDataprepFilter Interface
 
 	UPROPERTY()
-	UDataprepBoolFetcher* BoolFetcher;
+	TObjectPtr<UDataprepBoolFetcher> BoolFetcher;
 };
diff --git a/Engine/Plugins/Enterprise/DataprepEditor/Source/DataprepCore/Public/SelectionSystem/DataprepFloatFilter.h b/Engine/Plugins/Enterprise/DataprepEditor/Source/DataprepCore/Public/SelectionSystem/DataprepFloatFilter.h
--- a/Engine/Plugins/Enterprise/DataprepEditor/Source/DataprepCore/Public/SelectionSystem/DataprepFloatFilter.h
+++ b/Engine/Plugins/Enterprise/DataprepEditor/Source/DataprepCore/Public/SelectionSystem/DataprepFloatFilter.h
@@ -53,7 +53,7 @@ public:
 private:
 	// The source of float selected by the user
 	UPROPERTY()
-	UDataprepFloatFetcher* FloatFetcher;
+	TObjectPtr<UDataprepFloatFetcher> FloatFetcher;
 
 	// The matching criteria used when checking if a fetched value can pass the filter
 	UPROPERTY(EditAnywhere, Category = "Filter")
diff --git a/Engine/Plugins/Enterprise/DataprepEditor/Source/DataprepCore/Public/SelectionSystem/DataprepIntegerFilter.h b/Engine/Plugins/Enterprise/DataprepEditor/Source/DataprepCore/Public/SelectionSystem/DataprepIntegerFilter.h
--- a/Engine/Plugins/Enterprise/DataprepEditor/Source/DataprepCore/Public/SelectionSystem/DataprepIntegerFilter.h
+++ b/Engine/Plugins/Enterprise/DataprepEditor/Source/DataprepCore/Public/SelectionSystem/DataprepIntegerFilter.h
@@ -58,7 +58,7 @@ public:
 private:
 	// The source of int selected by the user
 	UPROPERTY()
-	UDataprepIntegerFetcher* IntFetcher;
+	TObjectPtr<UDataprepIntegerFetcher> IntFetcher;
 
 	// The matching criteria used when checking if a fetched value can pass the filter
 	UPROPERTY(EditAnywhere, Category = Filter)
diff --git a/Engine/Plugins/Enterprise/DataprepEditor/Source/DataprepCore/Public/SelectionSystem/DataprepStringFilter.h b/Engine/Plugins/Enterprise/DataprepEditor/Source/DataprepCore/Public/SelectionSystem/DataprepStringFilter.h
--- a/Engine/Plugins/Enterprise/DataprepEditor/Source/DataprepCore/Public/SelectionSystem/DataprepStringFilter.h
+++ b/Engine/Plugins/Enterprise/DataprepEditor/Source/DataprepCore/Public/SelectionSystem/DataprepStringFilter.h
@@ -76,12 +76,12 @@ private:
 	FString UserString;
 
 	UPROPERTY(EditAnywhere, Category = Filter)
-	UDataprepStringFilterMatchingArray* UserStringArray;
+	TObjectPtr<UDataprepStringFilterMatchingArray> UserStringArray;
 
 	UPROPERTY(EditAnywhere, Category = Filter)
 	bool bMatchInArray;
 
 	// The source of string selected by the user
 	UPROPERTY()
-	UDataprepStringFetcher* StringFetcher;
+	TObjectPtr<UDataprepStringFetcher> StringFetcher;
 };
diff --git a/Engine/Plugins/Enterprise/DataprepEditor/Source/DataprepCore/Public/SelectionSystem/DataprepStringsArrayFilter.h b/Engine/Plugins/Enterprise/DataprepEditor/Source/DataprepCore/Public/SelectionSystem/DataprepStringsArrayFilter.h
--- a/Engine/Plugins/Enterprise/DataprepEditor/Source/DataprepCore/Public/SelectionSystem/DataprepStringsArrayFilter.h
+++ b/Engine/Plugins/Enterprise/DataprepEditor/Source/DataprepCore/Public/SelectionSystem/DataprepStringsArrayFilter.h
@@ -50,12 +50,12 @@ private:
 	FString UserString;
 
 	UPROPERTY(EditAnywhere, Category = Filter)
-	UDataprepStringFilterMatchingArray* UserStringArray;
+	TObjectPtr<UDataprepStringFilterMatchingArray> UserStringArray;
 
 	UPROPERTY(EditAnywhere, Category = Filter)
 	bool bMatchInArray;
 
 	// The source of string selected by the user
 	UPROPERTY()
-	UDataprepStringsArrayFetcher* StringsArrayFetcher;
+	TObjectPtr<UDataprepStringsArrayFetcher> StringsArrayFetcher;
 };
diff --git a/Engine/Plugins/Enterprise/DataprepEditor/Source/DataprepEditor/Private/BlueprintNodes/K2Node_DataprepAction.h b/Engine/Plugins/Enterprise/DataprepEditor/Source/DataprepEditor/Private/BlueprintNodes/K2Node_DataprepAction.h
--- a/Engine/Plugins/Enterprise/DataprepEditor/Source/DataprepEditor/Private/BlueprintNodes/K2Node_DataprepAction.h
+++ b/Engine/Plugins/Enterprise/DataprepEditor/Source/DataprepEditor/Private/BlueprintNodes/K2Node_DataprepAction.h
@@ -49,7 +49,7 @@ protected:
 	FString ActionTitle;
 
 	UPROPERTY()
-	UDataprepActionAsset* DataprepAction_DEPRECATED;
+	TObjectPtr<UDataprepActionAsset> DataprepAction_DEPRECATED;
 
 	static const FName ThenPinName;
 	static const FName InObjectsPinName;
diff --git a/Engine/Plugins/Enterprise/DataprepEditor/Source/DataprepEditor/Private/DataprepGraph/DataprepGraphActionNode.h b/Engine/Plugins/Enterprise/DataprepEditor/Source/DataprepEditor/Private/DataprepGraph/DataprepGraphActionNode.h
--- a/Engine/Plugins/Enterprise/DataprepEditor/Source/DataprepEditor/Private/DataprepGraph/DataprepGraphActionNode.h
+++ b/Engine/Plugins/Enterprise/DataprepEditor/Source/DataprepEditor/Private/DataprepGraph/DataprepGraphActionNode.h
@@ -51,7 +51,7 @@ public:
 
 protected:
 	UPROPERTY()
-	UDataprepActionAsset* DataprepActionAsset;
+	TObjectPtr<UDataprepActionAsset> DataprepActionAsset;
 
 	UPROPERTY()
 	int32 StepIndex;
@@ -94,7 +94,7 @@ protected:
 	FString ActionTitle;
 
 	UPROPERTY()
-	UDataprepActionAsset* DataprepActionAsset;
+	TObjectPtr<UDataprepActionAsset> DataprepActionAsset;
 
 	UPROPERTY()
 	TWeakObjectPtr<UDataprepAsset> DataprepAssetPtr;
@@ -150,7 +150,7 @@ protected:
 	FString NodeTitle;
 
 	UPROPERTY()
-	TArray<UDataprepActionAsset*> Actions;
+	TArray<TObjectPtr<UDataprepActionAsset>> Actions;
 
 	UPROPERTY()
 	TWeakObjectPtr<UDataprepAsset> DataprepAssetPtr;
diff --git a/Engine/Plugins/Enterprise/DataprepEditor/Source/DataprepEditor/Public/DataprepEditorMenu.h b/Engine/Plugins/Enterprise/DataprepEditor/Source/DataprepEditor/Public/DataprepEditorMenu.h
--- a/Engine/Plugins/Enterprise/DataprepEditor/Source/DataprepEditor/Public/DataprepEditorMenu.h
+++ b/Engine/Plugins/Enterprise/DataprepEditor/Source/DataprepEditor/Public/DataprepEditorMenu.h
@@ -14,8 +14,8 @@ class DATAPREPEDITOR_API UDataprepEditorContextMenuContext : public UObject
 public:
 
 	UPROPERTY()
-	TArray< UObject* > SelectedObjects;
+	TArray< TObjectPtr<UObject> > SelectedObjects;
 
 	UPROPERTY()
-	class UDataprepAssetInterface* DataprepAsset;
+	TObjectPtr<class UDataprepAssetInterface> DataprepAsset;
 };
diff --git a/Engine/Plugins/Enterprise/DataprepEditor/Source/DataprepLibraries/Private/DataprepEditingOperations.h b/Engine/Plugins/Enterprise/DataprepEditor/Source/DataprepLibraries/Private/DataprepEditingOperations.h
--- a/Engine/Plugins/Enterprise/DataprepEditor/Source/DataprepLibraries/Private/DataprepEditingOperations.h
+++ b/Engine/Plugins/Enterprise/DataprepEditor/Source/DataprepLibraries/Private/DataprepEditingOperations.h
@@ -176,7 +176,7 @@ public:
 	void SetAsset(UObject* Asset);
 
 	UPROPERTY(EditAnywhere, BlueprintReadWrite, Category = ActorOperation)
-	UObject* SelectedAsset;
+	TObjectPtr<UObject> SelectedAsset;
 
 protected:
 	virtual void OnExecution_Implementation(const FDataprepContext& InContext) override;
diff --git a/Engine/Plugins/Enterprise/DataprepEditor/Source/DataprepLibraries/Private/DataprepOperations.h b/Engine/Plugins/Enterprise/DataprepEditor/Source/DataprepLibraries/Private/DataprepOperations.h
--- a/Engine/Plugins/Enterprise/DataprepEditor/Source/DataprepLibraries/Private/DataprepOperations.h
+++ b/Engine/Plugins/Enterprise/DataprepEditor/Source/DataprepLibraries/Private/DataprepOperations.h
@@ -196,7 +196,7 @@ class UDataprepSetMaterialOperation : public UDataprepOperation
 public:
 	// Material to use as a substitute
 	UPROPERTY(EditAnywhere, BlueprintReadWrite, Category = MeshOperation, meta = (ToolTip = "Material to use as a substitute"))
-	UMaterialInterface* Material;
+	TObjectPtr<UMaterialInterface> Material;
 
 	//~ Begin UDataprepOperation Interface
 public:
@@ -233,7 +233,7 @@ public:
 
 	// Material to use as a substitute
 	UPROPERTY(EditAnywhere, BlueprintReadWrite, Category = MeshOperation, meta = (ToolTip = "Material to use as a substitute"))
-	UMaterialInterface* MaterialSubstitute;
+	TObjectPtr<UMaterialInterface> MaterialSubstitute;
 
 	//~ Begin UDataprepOperation Interface
 public:
@@ -260,7 +260,7 @@ class UDataprepSubstituteMaterialByTableOperation : public UDataprepOperation
 public:
 	// Data table to use for the substitution
 	UPROPERTY(EditAnywhere, BlueprintReadWrite, Category = MeshOperation, meta = (ToolTip = "Data table to use for the substitution"))
-	UDataTable* MaterialDataTable;
+	TObjectPtr<UDataTable> MaterialDataTable;
 
 	//~ Begin UDataprepOperation Interface
 public:
@@ -312,7 +312,7 @@ class UDataprepSetMeshOperation : public UDataprepOperation
 public:
 	// Mesh to use as a substitute
 	UPROPERTY(EditAnywhere, BlueprintReadWrite, Category = ActorOperation, meta = (ToolTip = "Mesh to use as a substitute"))
-	UStaticMesh* StaticMesh;
+	TObjectPtr<UStaticMesh> StaticMesh;
 
 	//~ Begin UDataprepOperation Interface
 public:
diff --git a/Engine/Plugins/Enterprise/DataprepEditor/Source/DataprepLibraries/Private/DataprepOperationsLibrary.h b/Engine/Plugins/Enterprise/DataprepEditor/Source/DataprepLibraries/Private/DataprepOperationsLibrary.h
--- a/Engine/Plugins/Enterprise/DataprepEditor/Source/DataprepLibraries/Private/DataprepOperationsLibrary.h
+++ b/Engine/Plugins/Enterprise/DataprepEditor/Source/DataprepLibraries/Private/DataprepOperationsLibrary.h
@@ -37,7 +37,7 @@ struct FMaterialSubstitutionDataTable : public FTableRowBase
 
 	/** Material to use for the substitution */
 	UPROPERTY(EditAnywhere, BlueprintReadOnly, Category = "MaterialSubstitutionTable")
-	UMaterialInterface* MaterialReplacement = nullptr;
+	TObjectPtr<UMaterialInterface> MaterialReplacement = nullptr;
 };
 
 /*
@@ -58,7 +58,7 @@ struct FMeshSubstitutionDataTable
 
 	/** DEPRECATED - Mesh to use for the substitution */
 	UPROPERTY()
-	UStaticMesh* MeshReplacement_DEPRECATED = nullptr;
+	TObjectPtr<UStaticMesh> MeshReplacement_DEPRECATED = nullptr;
 };
 
 /*
diff --git a/Engine/Plugins/Enterprise/DatasmithImporter/Source/DatasmithImporter/Private/DatasmithFileProducer.h b/Engine/Plugins/Enterprise/DatasmithImporter/Source/DatasmithImporter/Private/DatasmithFileProducer.h
--- a/Engine/Plugins/Enterprise/DatasmithImporter/Source/DatasmithImporter/Private/DatasmithFileProducer.h
+++ b/Engine/Plugins/Enterprise/DatasmithImporter/Source/DatasmithImporter/Private/DatasmithFileProducer.h
@@ -99,15 +99,15 @@ private:
 	TArray<TStrongObjectPtr<UDatasmithOptionsBase>> ImportOptionsOverride;
 
 	UPROPERTY( Transient, DuplicateTransient )
-	UDatasmithScene* DatasmithScene;
+	TObjectPtr<UDatasmithScene> DatasmithScene;
 
 	UPROPERTY( Transient, DuplicateTransient )
-	UPackage* TransientPackage = nullptr;
+	TObjectPtr<UPackage> TransientPackage = nullptr;
 
 	TArray< TWeakObjectPtr< UObject > > Assets;
 
 	UPROPERTY()
-	TArray<UDatasmithOptionsBase*> TranslatorImportOptions;
+	TArray<TObjectPtr<UDatasmithOptionsBase>> TranslatorImportOptions;
 
 	UPROPERTY()
 	bool bTranslatorImportOptionsInitialized = false;
@@ -242,10 +242,10 @@ private:
 	TSet< FString > FilesToProcess;
 
 	UPROPERTY( Transient, DuplicateTransient )
-	UDatasmithFileProducer* FileProducer;
+	TObjectPtr<UDatasmithFileProducer> FileProducer;
 
 	UPROPERTY()
-	UDatasmithCommonTessellationOptions* TessellationOptions;
+	TObjectPtr<UDatasmithCommonTessellationOptions> TessellationOptions;
 
 	static TSet< FString > SupportedFormats;
 
diff --git a/Engine/Plugins/Enterprise/DatasmithImporter/Source/DatasmithImporter/Public/DatasmithBlueprintLibrary.h b/Engine/Plugins/Enterprise/DatasmithImporter/Source/DatasmithImporter/Public/DatasmithBlueprintLibrary.h
--- a/Engine/Plugins/Enterprise/DatasmithImporter/Source/DatasmithImporter/Public/DatasmithBlueprintLibrary.h
+++ b/Engine/Plugins/Enterprise/DatasmithImporter/Source/DatasmithImporter/Public/DatasmithBlueprintLibrary.h
@@ -28,17 +28,17 @@ struct DATASMITHIMPORTER_API FDatasmithImportFactoryCreateFileResult
 
 	/** List of all created actors if user requires to use multiple actors. The root actor will be the first element of the array. */
 	UPROPERTY(BlueprintReadWrite, Transient, Category = Imported)
-	TArray<class AActor*> ImportedActors;
+	TArray<TObjectPtr<class AActor>> ImportedActors;
 
 	/** Meshes created during the import process */
 	UPROPERTY(BlueprintReadWrite, Transient, Category = Imported)
-	TArray<class UObject*> ImportedMeshes;
+	TArray<TObjectPtr<class UObject>> ImportedMeshes;
 
 	UPROPERTY(BlueprintReadWrite, Transient, Category = Result)
 	bool bImportSucceed;
 
 	UPROPERTY(BlueprintReadWrite, Transient, Category = Result)
-	UDatasmithScene* Scene;
+	TObjectPtr<UDatasmithScene> Scene;
 };
 
 
diff --git a/Engine/Plugins/Enterprise/LidarPointCloud/Source/LidarPointCloudRuntime/Public/LidarPointCloud.h b/Engine/Plugins/Enterprise/LidarPointCloud/Source/LidarPointCloudRuntime/Public/LidarPointCloud.h
--- a/Engine/Plugins/Enterprise/LidarPointCloud/Source/LidarPointCloudRuntime/Public/LidarPointCloud.h
+++ b/Engine/Plugins/Enterprise/LidarPointCloud/Source/LidarPointCloudRuntime/Public/LidarPointCloud.h
@@ -156,9 +156,9 @@ private:
 
 	/** Description of collision */
 	UPROPERTY(transient, duplicatetransient)
-	UBodySetup* BodySetup;
+	TObjectPtr<UBodySetup> BodySetup;
 	UPROPERTY(transient, duplicatetransient)
-	UBodySetup* NewBodySetup;
+	TObjectPtr<UBodySetup> NewBodySetup;
 
 	/** Used for collision building */
 	FThreadSafeBool bCollisionBuildInProgress;
@@ -712,10 +712,10 @@ struct FLidarPointCloudTraceHit
 
 public:
 	UPROPERTY(EditAnywhere, BlueprintReadWrite, Category = "Collision")
-	ALidarPointCloudActor* Actor = nullptr;
+	TObjectPtr<ALidarPointCloudActor> Actor = nullptr;
 
 	UPROPERTY(EditAnywhere, BlueprintReadWrite, Category = "Collision")
-	ULidarPointCloudComponent* Component = nullptr;
+	TObjectPtr<ULidarPointCloudComponent> Component = nullptr;
 
 	UPROPERTY(EditAnywhere, BlueprintReadWrite, Category = "Collision")
 	TArray<FLidarPointCloudPoint> Points;
diff --git a/Engine/Plugins/Enterprise/LidarPointCloud/Source/LidarPointCloudRuntime/Public/LidarPointCloudActor.h b/Engine/Plugins/Enterprise/LidarPointCloud/Source/LidarPointCloudRuntime/Public/LidarPointCloudActor.h
--- a/Engine/Plugins/Enterprise/LidarPointCloud/Source/LidarPointCloudRuntime/Public/LidarPointCloudActor.h
+++ b/Engine/Plugins/Enterprise/LidarPointCloud/Source/LidarPointCloudRuntime/Public/LidarPointCloudActor.h
@@ -15,7 +15,7 @@ class LIDARPOINTCLOUDRUNTIME_API ALidarPointCloudActor : public AActor
 
 private:
 	UPROPERTY(Category = PointCloudActor, VisibleAnywhere, BlueprintReadOnly, meta = (ExposeFunctionCategories = "Rendering,Components|LidarPointCloud", AllowPrivateAccess = "true"))
-	ULidarPointCloudComponent* PointCloudComponent;
+	TObjectPtr<ULidarPointCloudComponent> PointCloudComponent;
 
 public:
 #if WITH_EDITOR
diff --git a/Engine/Plugins/Enterprise/LidarPointCloud/Source/LidarPointCloudRuntime/Public/LidarPointCloudComponent.h b/Engine/Plugins/Enterprise/LidarPointCloud/Source/LidarPointCloudRuntime/Public/LidarPointCloudComponent.h
--- a/Engine/Plugins/Enterprise/LidarPointCloud/Source/LidarPointCloudRuntime/Public/LidarPointCloudComponent.h
+++ b/Engine/Plugins/Enterprise/LidarPointCloud/Source/LidarPointCloudRuntime/Public/LidarPointCloudComponent.h
@@ -27,14 +27,14 @@ class LIDARPOINTCLOUDRUNTIME_API ULidarPointCloudComponent : public UMeshCompone
 		
 private:
 	UPROPERTY(EditAnywhere, BlueprintReadOnly, Category = "Lidar Point Cloud", meta = (AllowPrivateAccess = "true"))
-	ULidarPointCloud* PointCloud;
+	TObjectPtr<ULidarPointCloud> PointCloud;
 
 	/**
 	 * Allows using custom-built material for the point cloud.
 	 * Set to None to use the default one instead.
 	 */
 	UPROPERTY(EditAnywhere, BlueprintReadOnly, Category = "Material", meta = (AllowPrivateAccess = "true"))
-	UMaterialInterface* CustomMaterial;
+	TObjectPtr<UMaterialInterface> CustomMaterial;
 
 public:
 	/**
@@ -145,11 +145,11 @@ public:
 
 private:
 	UPROPERTY(Transient)
-	UMaterialInterface* Material;
+	TObjectPtr<UMaterialInterface> Material;
 	UPROPERTY(Transient)
-	UMaterialInterface* BaseMaterial;
+	TObjectPtr<UMaterialInterface> BaseMaterial;
 	UPROPERTY(Transient)
-	UMaterialInterface* BaseMaterialMasked;
+	TObjectPtr<UMaterialInterface> BaseMaterialMasked;
 
 	/** Pointer to the viewport client of the owning editor, or null, if this is a game object. */
 	TWeakPtr<FViewportClient> OwningViewportClient;
diff --git a/Engine/Plugins/Experimental/AR/PointCloud/Source/PointCloud/Public/PointCloudComponent.h b/Engine/Plugins/Experimental/AR/PointCloud/Source/PointCloud/Public/PointCloudComponent.h
--- a/Engine/Plugins/Experimental/AR/PointCloud/Source/PointCloud/Public/PointCloudComponent.h
+++ b/Engine/Plugins/Experimental/AR/PointCloud/Source/PointCloud/Public/PointCloudComponent.h
@@ -87,11 +87,11 @@ class POINTCLOUD_API UPointCloudComponent :
 
 	/** The material to render with */
 	UPROPERTY()
-	UMaterialInterface* PointCloudMaterial;
+	TObjectPtr<UMaterialInterface> PointCloudMaterial;
 	
 	/** An optional sprite texture to render the point with */
 	UPROPERTY(EditAnywhere, BlueprintReadOnly, Category = "Point Cloud")
-	UTexture* SpriteTexture;
+	TObjectPtr<UTexture> SpriteTexture;
 	
 private:
 	virtual void BeginPlay() override;
diff --git a/Engine/Plugins/Experimental/Animation/ContextualAnimation/Source/ContextualAnimation/Public/ContextualAnimManager.h b/Engine/Plugins/Experimental/Animation/ContextualAnimation/Source/ContextualAnimation/Public/ContextualAnimManager.h
--- a/Engine/Plugins/Experimental/Animation/ContextualAnimation/Source/ContextualAnimation/Public/ContextualAnimManager.h
+++ b/Engine/Plugins/Experimental/Animation/ContextualAnimation/Source/ContextualAnimation/Public/ContextualAnimManager.h
@@ -69,10 +69,10 @@ protected:
 
 	/** Container for all SceneActorComps in the world */
 	UPROPERTY()
-	TSet<UContextualAnimSceneActorComponent*> SceneActorCompContainer;
+	TSet<TObjectPtr<UContextualAnimSceneActorComponent>> SceneActorCompContainer;
 
 	UPROPERTY()
-	TArray<UContextualAnimSceneInstance*> Instances;
+	TArray<TObjectPtr<UContextualAnimSceneInstance>> Instances;
 
 	UFUNCTION()
 	void OnSceneInstanceEnded(UContextualAnimSceneInstance* SceneInstance);
diff --git a/Engine/Plugins/Experimental/Animation/ContextualAnimation/Source/ContextualAnimation/Public/ContextualAnimSceneActorComponent.h b/Engine/Plugins/Experimental/Animation/ContextualAnimation/Source/ContextualAnimation/Public/ContextualAnimSceneActorComponent.h
--- a/Engine/Plugins/Experimental/Animation/ContextualAnimation/Source/ContextualAnimation/Public/ContextualAnimSceneActorComponent.h
+++ b/Engine/Plugins/Experimental/Animation/ContextualAnimation/Source/ContextualAnimation/Public/ContextualAnimSceneActorComponent.h
@@ -33,7 +33,7 @@ public:
 	FContextualAnimSceneActorCompDelegate OnLeftSceneDelegate;
 
 	UPROPERTY(EditAnywhere, BlueprintReadWrite, Category = "Config")
-	class UContextualAnimSceneAsset* SceneAsset;
+	TObjectPtr<class UContextualAnimSceneAsset> SceneAsset;
 
 	UPROPERTY(EditAnywhere, BlueprintReadWrite, Category = "Debug")
 	bool bEnableDebug;
diff --git a/Engine/Plugins/Experimental/Animation/ContextualAnimation/Source/ContextualAnimation/Public/ContextualAnimTypes.h b/Engine/Plugins/Experimental/Animation/ContextualAnimation/Source/ContextualAnimation/Public/ContextualAnimTypes.h
--- a/Engine/Plugins/Experimental/Animation/ContextualAnimation/Source/ContextualAnimation/Public/ContextualAnimTypes.h
+++ b/Engine/Plugins/Experimental/Animation/ContextualAnimation/Source/ContextualAnimation/Public/ContextualAnimTypes.h
@@ -67,7 +67,7 @@ struct CONTEXTUALANIMATION_API FContextualAnimTrack
 	GENERATED_BODY()
 
 	UPROPERTY(EditAnywhere, BlueprintReadWrite, Category = "Defaults")
-	UAnimSequenceBase* Animation = nullptr;
+	TObjectPtr<UAnimSequenceBase> Animation = nullptr;
 
 	UPROPERTY(EditAnywhere, BlueprintReadWrite, Category = "Defaults")
 	float AnimMaxStartTime = 0.f;
@@ -84,7 +84,7 @@ struct CONTEXTUALANIMATION_API FContextualAnimTrack
 	FContextualAnimAlignmentTrackContainer IKTargetData;
 
 	UPROPERTY(EditAnywhere, Instanced, Category = "Defaults")
-	TArray<UContextualAnimSelectionCriterion*> SelectionCriteria;
+	TArray<TObjectPtr<UContextualAnimSelectionCriterion>> SelectionCriteria;
 
 	UPROPERTY(EditAnywhere, BlueprintReadWrite, Category = "Defaults")
 	FTransform MeshToScene;
diff --git a/Engine/Plugins/Experimental/Animation/ContextualAnimation/Source/ContextualAnimationEditor/Public/ContextualAnimEditorTypes.h b/Engine/Plugins/Experimental/Animation/ContextualAnimation/Source/ContextualAnimationEditor/Public/ContextualAnimEditorTypes.h
--- a/Engine/Plugins/Experimental/Animation/ContextualAnimation/Source/ContextualAnimationEditor/Public/ContextualAnimEditorTypes.h
+++ b/Engine/Plugins/Experimental/Animation/ContextualAnimation/Source/ContextualAnimationEditor/Public/ContextualAnimEditorTypes.h
@@ -71,7 +71,7 @@ struct FContextualAnimNewAnimSetData
 	FName RoleName = NAME_None;
 
 	UPROPERTY(EditAnywhere, Category = "Settings")
-	UAnimSequenceBase* Animation = nullptr;
+	TObjectPtr<UAnimSequenceBase> Animation = nullptr;
 
 	UPROPERTY(EditAnywhere, Category = "Settings")
 	bool bRequiresFlyingMode = false;
diff --git a/Engine/Plugins/Experimental/Animation/MotionWarping/Source/MotionWarping/Public/AnimNotifyState_MotionWarping.h b/Engine/Plugins/Experimental/Animation/MotionWarping/Source/MotionWarping/Public/AnimNotifyState_MotionWarping.h
--- a/Engine/Plugins/Experimental/Animation/MotionWarping/Source/MotionWarping/Public/AnimNotifyState_MotionWarping.h
+++ b/Engine/Plugins/Experimental/Animation/MotionWarping/Source/MotionWarping/Public/AnimNotifyState_MotionWarping.h
@@ -21,7 +21,7 @@ public:
 	//@TODO: Prevent notify callbacks and add comments explaining why we don't use those here.
 
 	UPROPERTY(EditAnywhere, Instanced, BlueprintReadWrite, Category = "Config")
-	URootMotionModifier* RootMotionModifier;
+	TObjectPtr<URootMotionModifier> RootMotionModifier;
 
 	UAnimNotifyState_MotionWarping(const FObjectInitializer& ObjectInitializer);
 
diff --git a/Engine/Plugins/Experimental/Animation/MotionWarping/Source/MotionWarping/Public/MotionWarpingComponent.h b/Engine/Plugins/Experimental/Animation/MotionWarping/Source/MotionWarping/Public/MotionWarpingComponent.h
--- a/Engine/Plugins/Experimental/Animation/MotionWarping/Source/MotionWarping/Public/MotionWarpingComponent.h
+++ b/Engine/Plugins/Experimental/Animation/MotionWarping/Source/MotionWarping/Public/MotionWarpingComponent.h
@@ -178,7 +178,7 @@ protected:
 
 	/** List of root motion modifiers */
 	UPROPERTY(Transient)
-	TArray<URootMotionModifier*> Modifiers;
+	TArray<TObjectPtr<URootMotionModifier>> Modifiers;
 
 	UPROPERTY(Transient, Replicated)
 	TArray<FMotionWarpingTarget> WarpTargets;
diff --git a/Engine/Plugins/Experimental/BlueprintSnapNodes/Source/BlueprintSnapNodes/Private/K2Node_SnapContainer.h b/Engine/Plugins/Experimental/BlueprintSnapNodes/Source/BlueprintSnapNodes/Private/K2Node_SnapContainer.h
--- a/Engine/Plugins/Experimental/BlueprintSnapNodes/Source/BlueprintSnapNodes/Private/K2Node_SnapContainer.h
+++ b/Engine/Plugins/Experimental/BlueprintSnapNodes/Source/BlueprintSnapNodes/Private/K2Node_SnapContainer.h
@@ -20,7 +20,7 @@ class UK2Node_SnapContainer : public UK2Node_Composite
 
 public:
 	UPROPERTY()
-	UK2Node* RootNode;
+	TObjectPtr<UK2Node> RootNode;
 
 public:
 	UK2Node_SnapContainer(const FObjectInitializer& ObjectInitializer);
diff --git a/Engine/Plugins/Experimental/ChaosCaching/Source/ChaosCaching/Public/Chaos/CacheCollection.h b/Engine/Plugins/Experimental/ChaosCaching/Source/ChaosCaching/Public/Chaos/CacheCollection.h
--- a/Engine/Plugins/Experimental/ChaosCaching/Source/ChaosCaching/Public/Chaos/CacheCollection.h
+++ b/Engine/Plugins/Experimental/ChaosCaching/Source/ChaosCaching/Public/Chaos/CacheCollection.h
@@ -24,5 +24,5 @@ public:
 	const TArray<UChaosCache*> GetCaches() const;
 
 	UPROPERTY(EditAnywhere, Instanced, Category="Caching", meta=(EditFixedOrder))
-	TArray<UChaosCache*> Caches;
+	TArray<TObjectPtr<UChaosCache>> Caches;
 };
\ No newline at end of file
diff --git a/Engine/Plugins/Experimental/ChaosCaching/Source/ChaosCaching/Public/Chaos/CacheEvents.h b/Engine/Plugins/Experimental/ChaosCaching/Source/ChaosCaching/Public/Chaos/CacheEvents.h
--- a/Engine/Plugins/Experimental/ChaosCaching/Source/ChaosCaching/Public/Chaos/CacheEvents.h
+++ b/Engine/Plugins/Experimental/ChaosCaching/Source/ChaosCaching/Public/Chaos/CacheEvents.h
@@ -61,7 +61,7 @@ struct FCacheEventTrack
 
 	/** Type of the event stored in this track. Must inherit FCacheEventBase */
 	UPROPERTY()
-	UScriptStruct* Struct;
+	TObjectPtr<UScriptStruct> Struct;
 
 	UPROPERTY()
 	TArray<float> TimeStamps;
diff --git a/Engine/Plugins/Experimental/ChaosCaching/Source/ChaosCaching/Public/Chaos/CacheManagerActor.h b/Engine/Plugins/Experimental/ChaosCaching/Source/ChaosCaching/Public/Chaos/CacheManagerActor.h
--- a/Engine/Plugins/Experimental/ChaosCaching/Source/ChaosCaching/Public/Chaos/CacheManagerActor.h
+++ b/Engine/Plugins/Experimental/ChaosCaching/Source/ChaosCaching/Public/Chaos/CacheManagerActor.h
@@ -114,7 +114,7 @@ public:
 	 * across multiple components depending on the settings for that component.
 	 */
 	UPROPERTY(EditInstanceOnly, Category = "Caching")
-	UChaosCacheCollection* CacheCollection;
+	TObjectPtr<UChaosCacheCollection> CacheCollection;
 
 	/** How to use the cache - playback or record */
 	UPROPERTY(EditAnywhere, Category = "Caching")
diff --git a/Engine/Plugins/Experimental/ChaosCaching/Source/ChaosCaching/Public/Chaos/ChaosCache.h b/Engine/Plugins/Experimental/ChaosCaching/Source/ChaosCaching/Public/Chaos/ChaosCache.h
--- a/Engine/Plugins/Experimental/ChaosCaching/Source/ChaosCaching/Public/Chaos/ChaosCache.h
+++ b/Engine/Plugins/Experimental/ChaosCaching/Source/ChaosCaching/Public/Chaos/ChaosCache.h
@@ -98,7 +98,7 @@ struct FCacheSpawnableTemplate
 	{}
 
 	UPROPERTY(VisibleAnywhere, Category = "Caching")
-	UObject* DuplicatedTemplate;
+	TObjectPtr<UObject> DuplicatedTemplate;
 
 	UPROPERTY(VisibleAnywhere, Category = "Caching")
 	FTransform InitialTransform;
diff --git a/Engine/Plugins/Experimental/ChaosVehiclesPlugin/Source/ChaosVehicles/Public/ChaosVehicleMovementComponent.h b/Engine/Plugins/Experimental/ChaosVehiclesPlugin/Source/ChaosVehicles/Public/ChaosVehicleMovementComponent.h
--- a/Engine/Plugins/Experimental/ChaosVehiclesPlugin/Source/ChaosVehicles/Public/ChaosVehicleMovementComponent.h
+++ b/Engine/Plugins/Experimental/ChaosVehiclesPlugin/Source/ChaosVehicles/Public/ChaosVehicleMovementComponent.h
@@ -1342,7 +1342,7 @@ protected:
 	TUniquePtr<UChaosVehicleSimulation> VehicleSimulationPT;	/* simulation code running on the physics thread async callback */
 
 	UPROPERTY(transient, Replicated)
-	AController* OverrideController;
+	TObjectPtr<AController> OverrideController;
 
 	const Chaos::FSimpleAerodynamicsConfig& GetAerodynamicsConfig()
 	{
diff --git a/Engine/Plugins/Experimental/ChaosVehiclesPlugin/Source/ChaosVehicles/Public/ChaosVehicleWheel.h b/Engine/Plugins/Experimental/ChaosVehiclesPlugin/Source/ChaosVehicles/Public/ChaosVehicleWheel.h
--- a/Engine/Plugins/Experimental/ChaosVehiclesPlugin/Source/ChaosVehicles/Public/ChaosVehicleWheel.h
+++ b/Engine/Plugins/Experimental/ChaosVehiclesPlugin/Source/ChaosVehicles/Public/ChaosVehicleWheel.h
@@ -65,7 +65,7 @@ class UChaosWheeledVehicleMovementComponent;
 			* (if empty, sphere will be added as wheel shape, check bDontCreateShape flag)
 			*/
 		UPROPERTY(EditDefaultsOnly, Category = Shape)
-		class UStaticMesh* CollisionMesh;
+		TObjectPtr<class UStaticMesh> CollisionMesh;
 
 		/** If left undefined then the bAffectedByEngine value is used, if defined then bAffectedByEngine is ignored and the differential setup on the vehicle defines which wheels get power from the engine */
 		UPROPERTY(EditAnywhere, Category = Wheel)
@@ -203,7 +203,7 @@ class UChaosWheeledVehicleMovementComponent;
 
 		/** The vehicle that owns us */
 		UPROPERTY(transient)
-		class UChaosWheeledVehicleMovementComponent* VehicleComponent;
+		TObjectPtr<class UChaosWheeledVehicleMovementComponent> VehicleComponent;
 
 		// Our index in the vehicle's (and setup's) wheels array
 		UPROPERTY(transient)
diff --git a/Engine/Plugins/Experimental/ChaosVehiclesPlugin/Source/ChaosVehicles/Public/ChaosWheeledVehicleMovementComponent.h b/Engine/Plugins/Experimental/ChaosVehiclesPlugin/Source/ChaosVehicles/Public/ChaosWheeledVehicleMovementComponent.h
--- a/Engine/Plugins/Experimental/ChaosVehiclesPlugin/Source/ChaosVehicles/Public/ChaosWheeledVehicleMovementComponent.h
+++ b/Engine/Plugins/Experimental/ChaosVehiclesPlugin/Source/ChaosVehicles/Public/ChaosWheeledVehicleMovementComponent.h
@@ -655,7 +655,7 @@ class CHAOSVEHICLES_API UChaosWheeledVehicleMovementComponent : public UChaosVeh
 
 	// Our instanced wheels
 	UPROPERTY(transient, duplicatetransient, BlueprintReadOnly, Category = Vehicle)
-	TArray<class UChaosVehicleWheel*> Wheels;
+	TArray<TObjectPtr<class UChaosVehicleWheel>> Wheels;
 
 	/** Get current engine's rotation speed */
 	UFUNCTION(BlueprintCallable, Category = "Game|Components|ChaosWheeledVehicleMovement")
diff --git a/Engine/Plugins/Experimental/ChaosVehiclesPlugin/Source/ChaosVehicles/Public/VehicleAnimationInstance.h b/Engine/Plugins/Experimental/ChaosVehiclesPlugin/Source/ChaosVehicles/Public/VehicleAnimationInstance.h
--- a/Engine/Plugins/Experimental/ChaosVehiclesPlugin/Source/ChaosVehicles/Public/VehicleAnimationInstance.h
+++ b/Engine/Plugins/Experimental/ChaosVehiclesPlugin/Source/ChaosVehicles/Public/VehicleAnimationInstance.h
@@ -106,7 +106,7 @@ private:
 	FVehicleAnimationInstanceProxy AnimInstanceProxy;
 
 	UPROPERTY(transient)
-	const UChaosWheeledVehicleMovementComponent* WheeledVehicleComponent;
+	TObjectPtr<const UChaosWheeledVehicleMovementComponent> WheeledVehicleComponent;
 };
 
 
diff --git a/Engine/Plugins/Experimental/ChaosVehiclesPlugin/Source/ChaosVehicles/Public/WheeledVehiclePawn.h b/Engine/Plugins/Experimental/ChaosVehiclesPlugin/Source/ChaosVehicles/Public/WheeledVehiclePawn.h
--- a/Engine/Plugins/Experimental/ChaosVehiclesPlugin/Source/ChaosVehicles/Public/WheeledVehiclePawn.h
+++ b/Engine/Plugins/Experimental/ChaosVehiclesPlugin/Source/ChaosVehicles/Public/WheeledVehiclePawn.h
@@ -24,11 +24,11 @@ class CHAOSVEHICLES_API AWheeledVehiclePawn : public APawn
 private:
 	/**  The main skeletal mesh associated with this Vehicle */
 	UPROPERTY(Category = Vehicle, VisibleDefaultsOnly, BlueprintReadOnly, meta = (AllowPrivateAccess = "true"))
-	class USkeletalMeshComponent* Mesh;
+	TObjectPtr<class USkeletalMeshComponent> Mesh;
 
 	/** vehicle simulation component */
 	UPROPERTY(Category = Vehicle, VisibleDefaultsOnly, BlueprintReadOnly, meta = (AllowPrivateAccess = "true"))
-	class UChaosVehicleMovementComponent* VehicleMovementComponent;
+	TObjectPtr<class UChaosVehicleMovementComponent> VehicleMovementComponent;
 public:
 
 	/** Name of the MeshComponent. Use this name if you want to prevent creation of the component (with ObjectInitializer.DoNotCreateDefaultSubobject). */
diff --git a/Engine/Plugins/Experimental/CinematicPrestreaming/Source/CinematicPrestreamingEditor/Private/CinePrestreamingRecorderSetting.h b/Engine/Plugins/Experimental/CinematicPrestreaming/Source/CinematicPrestreamingEditor/Private/CinePrestreamingRecorderSetting.h
--- a/Engine/Plugins/Experimental/CinematicPrestreaming/Source/CinematicPrestreamingEditor/Private/CinePrestreamingRecorderSetting.h
+++ b/Engine/Plugins/Experimental/CinematicPrestreaming/Source/CinematicPrestreamingEditor/Private/CinePrestreamingRecorderSetting.h
@@ -15,10 +15,10 @@ struct FMoviePipelineCinePrestreamingGeneratedData
 	GENERATED_BODY()
 
 	UPROPERTY(EditAnywhere, BlueprintReadWrite, Category = "Cinematic Prestreaming")
-	UCinePrestreamingData* StreamingData;
+	TObjectPtr<UCinePrestreamingData> StreamingData;
 
 	UPROPERTY(EditAnywhere, BlueprintReadWrite, Category = "Cinematic Prestreaming")
-	UMovieScene* MovieScene;
+	TObjectPtr<UMovieScene> MovieScene;
 
 	UPROPERTY(EditAnywhere, BlueprintReadWrite, Category = "Cinematic Prestreaming")
 	FString PackagePath;
diff --git a/Engine/Plugins/Experimental/CinematicPrestreaming/Source/CinematicPrestreamingEditor/Public/CinePrestreamingEditorSubsystem.h b/Engine/Plugins/Experimental/CinematicPrestreaming/Source/CinematicPrestreamingEditor/Public/CinePrestreamingEditorSubsystem.h
--- a/Engine/Plugins/Experimental/CinematicPrestreaming/Source/CinematicPrestreamingEditor/Public/CinePrestreamingEditorSubsystem.h
+++ b/Engine/Plugins/Experimental/CinematicPrestreaming/Source/CinematicPrestreamingEditor/Public/CinePrestreamingEditorSubsystem.h
@@ -66,7 +66,7 @@ private:
 	// UMoviePipelineQueue* OriginalQueue;
 
 	UPROPERTY(Transient)
-	UMoviePipelineExecutorBase* ActiveExecutor;
+	TObjectPtr<UMoviePipelineExecutorBase> ActiveExecutor;
 
 	TOptional<FCinePrestreamingGenerateAssetArgs> ActiveAssetArgs;
 };
\ No newline at end of file
diff --git a/Engine/Plugins/Experimental/CodeEditor/Source/CodeEditor/Private/CodeProjectItem.h b/Engine/Plugins/Experimental/CodeEditor/Source/CodeEditor/Private/CodeProjectItem.h
--- a/Engine/Plugins/Experimental/CodeEditor/Source/CodeEditor/Private/CodeProjectItem.h
+++ b/Engine/Plugins/Experimental/CodeEditor/Source/CodeEditor/Private/CodeProjectItem.h
@@ -47,7 +47,7 @@ public:
 	FString Path;
 
 	UPROPERTY(Transient)
-	TArray<UCodeProjectItem*> Children;
+	TArray<TObjectPtr<UCodeProjectItem>> Children;
 
 	/** Delegate handle for directory watcher */
 	FDelegateHandle OnDirectoryChangedHandle;
diff --git a/Engine/Plugins/Experimental/CommonConversation/Source/CommonConversationRuntime/Public/ConversationContext.h b/Engine/Plugins/Experimental/CommonConversation/Source/CommonConversationRuntime/Public/ConversationContext.h
--- a/Engine/Plugins/Experimental/CommonConversation/Source/CommonConversationRuntime/Public/ConversationContext.h
+++ b/Engine/Plugins/Experimental/CommonConversation/Source/CommonConversationRuntime/Public/ConversationContext.h
@@ -213,16 +213,16 @@ public:
 
 private:
 	UPROPERTY()
-	UConversationRegistry* ConversationRegistry = nullptr;
+	TObjectPtr<UConversationRegistry> ConversationRegistry = nullptr;
 
 	UPROPERTY()
-	UConversationInstance* ActiveConversation = nullptr;
+	TObjectPtr<UConversationInstance> ActiveConversation = nullptr;
 
 	UPROPERTY()
-	UConversationParticipantComponent* ClientParticipant = nullptr;
+	TObjectPtr<UConversationParticipantComponent> ClientParticipant = nullptr;
 
 	UPROPERTY()
-	const UConversationTaskNode* TaskBeingConsidered = nullptr;
+	TObjectPtr<const UConversationTaskNode> TaskBeingConsidered = nullptr;
 
 	UPROPERTY()
 	TArray<FConversationNodeHandle> ReturnScopeStack;
diff --git a/Engine/Plugins/Experimental/CommonConversation/Source/CommonConversationRuntime/Public/ConversationDatabase.h b/Engine/Plugins/Experimental/CommonConversation/Source/CommonConversationRuntime/Public/ConversationDatabase.h
--- a/Engine/Plugins/Experimental/CommonConversation/Source/CommonConversationRuntime/Public/ConversationDatabase.h
+++ b/Engine/Plugins/Experimental/CommonConversation/Source/CommonConversationRuntime/Public/ConversationDatabase.h
@@ -96,7 +96,7 @@ private: // Compiled Data
 
 	// Compiled: Reachable nodes
 	UPROPERTY()
-	TMap<FGuid, UConversationNode*> ReachableNodeMap;
+	TMap<FGuid, TObjectPtr<UConversationNode>> ReachableNodeMap;
 
 	// Compiled: Entry points
 	UPROPERTY(AssetRegistrySearchable)
@@ -124,11 +124,11 @@ private:
 #if WITH_EDITORONLY_DATA
 	// All nodes
 	UPROPERTY()
-	TMap<FGuid, UConversationNode*> FullNodeMap;
+	TMap<FGuid, TObjectPtr<UConversationNode>> FullNodeMap;
 
 	// 'Source code' graphs (of type UConversationGraph)
 	UPROPERTY()
-	TArray<UEdGraph*> SourceGraphs;
+	TArray<TObjectPtr<UEdGraph>> SourceGraphs;
 
 public:
 	// Info about the graphs we last edited
diff --git a/Engine/Plugins/Experimental/CommonConversation/Source/CommonConversationRuntime/Public/ConversationNode.h b/Engine/Plugins/Experimental/CommonConversation/Source/CommonConversationRuntime/Public/ConversationNode.h
--- a/Engine/Plugins/Experimental/CommonConversation/Source/CommonConversationRuntime/Public/ConversationNode.h
+++ b/Engine/Plugins/Experimental/CommonConversation/Source/CommonConversationRuntime/Public/ConversationNode.h
@@ -134,7 +134,7 @@ protected:
 
 public:
 	UPROPERTY(Transient, DuplicateTransient)
-	mutable UObject* EvalWorldContextObj;
+	mutable TObjectPtr<UObject> EvalWorldContextObj;
 
 protected:
 	/** node name */
@@ -160,7 +160,7 @@ protected:
 private:
 	/** parent node */
 	UPROPERTY()
-	UConversationNode* ParentNode;
+	TObjectPtr<UConversationNode> ParentNode;
 
 private:
 	friend class FConversationCompiler;
diff --git a/Engine/Plugins/Experimental/CommonConversation/Source/CommonConversationRuntime/Public/ConversationParticipantComponent.h b/Engine/Plugins/Experimental/CommonConversation/Source/CommonConversationRuntime/Public/ConversationParticipantComponent.h
--- a/Engine/Plugins/Experimental/CommonConversation/Source/CommonConversationRuntime/Public/ConversationParticipantComponent.h
+++ b/Engine/Plugins/Experimental/CommonConversation/Source/CommonConversationRuntime/Public/ConversationParticipantComponent.h
@@ -146,10 +146,10 @@ private:
 
 private:
 	UPROPERTY()
-	UConversationInstance* Auth_CurrentConversation;
+	TObjectPtr<UConversationInstance> Auth_CurrentConversation;
 
 	UPROPERTY()
-	TArray<UConversationInstance*> Auth_Conversations;
+	TArray<TObjectPtr<UConversationInstance>> Auth_Conversations;
 	
 	UPROPERTY()
     FClientConversationMessagePayload LastMessage;
diff --git a/Engine/Plugins/Experimental/CommonConversation/Source/CommonConversationRuntime/Public/ConversationRegistry.h b/Engine/Plugins/Experimental/CommonConversation/Source/CommonConversationRuntime/Public/ConversationRegistry.h
--- a/Engine/Plugins/Experimental/CommonConversation/Source/CommonConversationRuntime/Public/ConversationRegistry.h
+++ b/Engine/Plugins/Experimental/CommonConversation/Source/CommonConversationRuntime/Public/ConversationRegistry.h
@@ -28,10 +28,10 @@ struct COMMONCONVERSATIONRUNTIME_API FNetSerializeScriptStructCache_ConvVersion
 	bool NetSerialize(FArchive& Ar, UScriptStruct*& Struct);
 
 	UPROPERTY()
-	TMap<UScriptStruct*, int32> ScriptStructsToIndex;
+	TMap<TObjectPtr<UScriptStruct>, int32> ScriptStructsToIndex;
 
 	UPROPERTY()
-	TArray<UScriptStruct*> IndexToScriptStructs;
+	TArray<TObjectPtr<UScriptStruct>> IndexToScriptStructs;
 };
 
 /**
diff --git a/Engine/Plugins/Experimental/CommonConversation/Source/CommonConversationRuntime/Public/ConversationTaskNode.h b/Engine/Plugins/Experimental/CommonConversation/Source/CommonConversationRuntime/Public/ConversationTaskNode.h
--- a/Engine/Plugins/Experimental/CommonConversation/Source/CommonConversationRuntime/Public/ConversationTaskNode.h
+++ b/Engine/Plugins/Experimental/CommonConversation/Source/CommonConversationRuntime/Public/ConversationTaskNode.h
@@ -43,7 +43,7 @@ class COMMONCONVERSATIONRUNTIME_API UConversationTaskNode : public UConversation
 public:
 	// Requirements and side effects
 	UPROPERTY()
-	TArray<UConversationSubNode*> SubNodes;
+	TArray<TObjectPtr<UConversationSubNode>> SubNodes;
 
 #if WITH_EDITORONLY_DATA
 	/* EDITOR ONLY VISUALS: Does this task internally have requirements? */
diff --git a/Engine/Plugins/Experimental/CommonConversation/Source/CommonConversationRuntime/Public/ConversationTypes.h b/Engine/Plugins/Experimental/CommonConversation/Source/CommonConversationRuntime/Public/ConversationTypes.h
--- a/Engine/Plugins/Experimental/CommonConversation/Source/CommonConversationRuntime/Public/ConversationTypes.h
+++ b/Engine/Plugins/Experimental/CommonConversation/Source/CommonConversationRuntime/Public/ConversationTypes.h
@@ -155,7 +155,7 @@ struct COMMONCONVERSATIONRUNTIME_API FConversationParticipantEntry
 public:
 
 	UPROPERTY(BlueprintReadWrite, Category=Conversation)
-	AActor* Actor = nullptr;
+	TObjectPtr<AActor> Actor = nullptr;
 
 	UPROPERTY(BlueprintReadWrite, Category=Conversation)
 	FGameplayTag ParticipantID;
diff --git a/Engine/Plugins/Experimental/CommonUI/Source/CommonInput/Public/CommonInputActionDomain.h b/Engine/Plugins/Experimental/CommonUI/Source/CommonInput/Public/CommonInputActionDomain.h
--- a/Engine/Plugins/Experimental/CommonUI/Source/CommonInput/Public/CommonInputActionDomain.h
+++ b/Engine/Plugins/Experimental/CommonUI/Source/CommonInput/Public/CommonInputActionDomain.h
@@ -65,5 +65,5 @@ class COMMONINPUT_API UCommonInputActionDomainTable : public UDataAsset
 public:
 	// Domains will receive events in ascending index order
 	UPROPERTY(EditDefaultsOnly, BlueprintReadOnly, Category = "Default")
-	TArray<UCommonInputActionDomain*> ActionDomains;
+	TArray<TObjectPtr<UCommonInputActionDomain>> ActionDomains;
 };
\ No newline at end of file
diff --git a/Engine/Plugins/Experimental/CommonUI/Source/CommonUI/Public/CommonActionWidget.h b/Engine/Plugins/Experimental/CommonUI/Source/CommonUI/Public/CommonActionWidget.h
--- a/Engine/Plugins/Experimental/CommonUI/Source/CommonUI/Public/CommonActionWidget.h
+++ b/Engine/Plugins/Experimental/CommonUI/Source/CommonUI/Public/CommonActionWidget.h
@@ -109,7 +109,7 @@ protected:
 	void HandleInputMethodChanged(ECommonInputType InInputType);
 
 	UPROPERTY(Transient)
-	UMaterialInstanceDynamic* ProgressDynamicMaterial;
+	TObjectPtr<UMaterialInstanceDynamic> ProgressDynamicMaterial;
 	
 	TSharedPtr<SBox> MyKeyBox;
 
diff --git a/Engine/Plugins/Experimental/CommonUI/Source/CommonUI/Public/CommonButtonBase.h b/Engine/Plugins/Experimental/CommonUI/Source/CommonUI/Public/CommonButtonBase.h
--- a/Engine/Plugins/Experimental/CommonUI/Source/CommonUI/Public/CommonButtonBase.h
+++ b/Engine/Plugins/Experimental/CommonUI/Source/CommonUI/Public/CommonButtonBase.h
@@ -716,7 +716,7 @@ private:
 
 	/** The dynamic material instance of the material set by the single material style, if specified. */
 	UPROPERTY(Transient)
-	UMaterialInstanceDynamic* SingleMaterialStyleMID;
+	TObjectPtr<UMaterialInstanceDynamic> SingleMaterialStyleMID;
 
 	/** Internally managed and applied style to use when not selected */
 	UPROPERTY()
@@ -758,5 +758,5 @@ protected:
 	 * Visualize the default click action while hovered
 	 */
 	UPROPERTY(BlueprintReadOnly, Category = Input, meta = (BindWidget, OptionalWidget = true, AllowPrivateAccess = true))
-	UCommonActionWidget* InputActionWidget;
+	TObjectPtr<UCommonActionWidget> InputActionWidget;
 };
\ No newline at end of file
diff --git a/Engine/Plugins/Experimental/CommonUI/Source/CommonUI/Public/CommonLazyWidget.h b/Engine/Plugins/Experimental/CommonUI/Source/CommonUI/Public/CommonLazyWidget.h
--- a/Engine/Plugins/Experimental/CommonUI/Source/CommonUI/Public/CommonLazyWidget.h
+++ b/Engine/Plugins/Experimental/CommonUI/Source/CommonUI/Public/CommonLazyWidget.h
@@ -60,7 +60,7 @@ private:
 	FSlateBrush LoadingBackgroundBrush;
 
 	UPROPERTY(Transient)
-	UUserWidget* Content;
+	TObjectPtr<UUserWidget> Content;
 
 	TSharedPtr<FStreamableHandle> StreamingHandle;
 	FSoftObjectPath StreamingObjectPath;
diff --git a/Engine/Plugins/Experimental/CommonUI/Source/CommonUI/Public/CommonRotator.h b/Engine/Plugins/Experimental/CommonUI/Source/CommonUI/Public/CommonRotator.h
--- a/Engine/Plugins/Experimental/CommonUI/Source/CommonUI/Public/CommonRotator.h
+++ b/Engine/Plugins/Experimental/CommonUI/Source/CommonUI/Public/CommonRotator.h
@@ -78,7 +78,7 @@ protected:
 
 	/** The displayed text */
 	UPROPERTY(BlueprintReadOnly, Category = CommonRotator, Meta = (BindWidget))
-	UCommonTextBlock* MyText;
+	TObjectPtr<UCommonTextBlock> MyText;
 
 	// Holds the array of display texts
 	TArray<FText> TextLabels;
diff --git a/Engine/Plugins/Experimental/CommonUI/Source/CommonUI/Public/CommonTabListWidgetBase.h b/Engine/Plugins/Experimental/CommonUI/Source/CommonUI/Public/CommonTabListWidgetBase.h
--- a/Engine/Plugins/Experimental/CommonUI/Source/CommonUI/Public/CommonTabListWidgetBase.h
+++ b/Engine/Plugins/Experimental/CommonUI/Source/CommonUI/Public/CommonTabListWidgetBase.h
@@ -23,11 +23,11 @@ public:
 	
 	/** The actual button widget that represents this tab on-screen */
 	UPROPERTY()
-	UCommonButtonBase* TabButton;
+	TObjectPtr<UCommonButtonBase> TabButton;
 
 	/** The actual instance of the content widget to display when this tab is selected. Can be null if a load is required. */
 	UPROPERTY()
-	UWidget* ContentInstance;
+	TObjectPtr<UWidget> ContentInstance;
 
 	FCommonRegisteredTabInfo()
 		: TabIndex(INDEX_NONE)
@@ -206,7 +206,7 @@ protected:
 	
 	/** The button group that manages all the created tab buttons */
 	UPROPERTY(Transient)
-	UCommonButtonGroupBase* TabButtonGroup;
+	TObjectPtr<UCommonButtonGroupBase> TabButtonGroup;
 
 	/** Is the tab list currently listening for tab input actions? */
 	bool bIsListeningForInput = false;
