Removing deprecated lens functionality from Composure.

#rnx

diff --git a/Engine/Plugins/Compositing/Composure/Source/Composure/Classes/CompositingCaptureBase.h b/Engine/Plugins/Compositing/Composure/Source/Composure/Classes/CompositingCaptureBase.h
--- a/Engine/Plugins/Compositing/Composure/Source/Composure/Classes/CompositingCaptureBase.h
+++ b/Engine/Plugins/Compositing/Composure/Source/Composure/Classes/CompositingCaptureBase.h
@@ -47,14 +47,6 @@ protected:
 	UPROPERTY()
 	TObjectPtr<UMaterialInstanceDynamic> LastDistortionMID = nullptr;
 
-#if WITH_EDITORONLY_DATA
-PRAGMA_DISABLE_DEPRECATION_WARNINGS
-	UE_DEPRECATED(5.1, "This property has been deprecated. Use the LensComponent picker to choose a lens to use for distortion.")
-	UPROPERTY()
-	FDistortionHandlerPicker DistortionSource_DEPRECATED;
-PRAGMA_ENABLE_DEPRECATION_WARNINGS
-#endif //WITH_EDITORONLY_DATA
-
 public:
 	/** Default constructor */
 	ACompositingCaptureBase();
@@ -76,10 +68,4 @@ public:
 
 	/** Set the Lens Component this CG layer will use to drive distortion on the scene capture */
 	void SetLens(ULensComponent* InLens);
-
-	UE_DEPRECATED(5.1, "This function has been deprecated. Use SetLens() to set a lens component that this CG layer should use to query for distortion.")
-	void SetDistortionHandler(ULensDistortionModelHandlerBase* InDistortionHandler);
-
-	UE_DEPRECATED(5.1, "This function has been deprecated. Query the lens component used by this CG layer for its distortion handler instead.")
-	ULensDistortionModelHandlerBase* GetDistortionHandler();
 };
diff --git a/Engine/Plugins/Compositing/Composure/Source/Composure/Private/CompositingCaptureBase.cpp b/Engine/Plugins/Compositing/Composure/Source/Composure/Private/CompositingCaptureBase.cpp
--- a/Engine/Plugins/Compositing/Composure/Source/Composure/Private/CompositingCaptureBase.cpp
+++ b/Engine/Plugins/Compositing/Composure/Source/Composure/Private/CompositingCaptureBase.cpp
@@ -9,7 +9,6 @@
 
 DEFINE_LOG_CATEGORY_STATIC(LogCompositingCaptureBase, Log, All);
 
-PRAGMA_DISABLE_DEPRECATION_WARNINGS
 ACompositingCaptureBase::ACompositingCaptureBase()
 {
 	// Create the SceneCapture component and assign its parent to be the RootComponent (the ComposurePostProcessingPassProxy)
@@ -19,7 +18,6 @@ ACompositingCaptureBase::ACompositingCaptureBase()
 	// The SceneCaptureComponent2D default constructor disables TAA, but CG Compositing Elements enable it by default
 	SceneCaptureComponent2D->ShowFlags.TemporalAA = true;
 }
-PRAGMA_ENABLE_DEPRECATION_WARNINGS
 
 void ACompositingCaptureBase::PostInitProperties()
 {
@@ -167,24 +165,3 @@ void ACompositingCaptureBase::SetLens(ULensComponent* InLens)
 		LensComponentPicker.PathToComponent = InLens->GetPathName(InLens->GetOwner());
 	}
 }
-
-void ACompositingCaptureBase::SetDistortionHandler(ULensDistortionModelHandlerBase* InDistortionHandler)
-{
-	// This function has been deprecated.
-}
-
-ULensDistortionModelHandlerBase* ACompositingCaptureBase::GetDistortionHandler()
-{
-	ULensDistortionModelHandlerBase* LensDistortionHandler = nullptr;
-
-	if (ACameraActor* TargetCamera = FindTargetCamera())
-	{
-		ULensComponent* LensComponent = Cast<ULensComponent>(LensComponentPicker.GetComponent(TargetCamera));
-		if (LensComponent && LensComponent->GetOwner() == TargetCamera)
-		{
-			LensDistortionHandler = LensComponent->GetLensDistortionHandler();
-		}
-	}
-
-	return LensDistortionHandler;
-}
