Composure: Fix scene view extension now that scene capture correctly call BeginRenderViewFamily.

#rnx

diff --git a/Engine/Plugins/Compositing/Composure/Source/Composure/Private/ComposureViewExtension.cpp b/Engine/Plugins/Compositing/Composure/Source/Composure/Private/ComposureViewExtension.cpp
--- a/Engine/Plugins/Compositing/Composure/Source/Composure/Private/ComposureViewExtension.cpp
+++ b/Engine/Plugins/Compositing/Composure/Source/Composure/Private/ComposureViewExtension.cpp
@@ -15,18 +15,32 @@ void FComposureViewExtension::BeginRenderViewFamily(FSceneViewFamily& InViewFami
 {
 	if (AssociatedPipelineObj.IsValid())
 	{
+		bool bHasSceneCaptureView = false;
 		bool bCameraCutThisFrame = false;
 		for (const FSceneView* SceneView : InViewFamily.Views)
 		{
-			if (SceneView && SceneView->bCameraCut)
+			if (SceneView)
 			{
-				bCameraCutThisFrame = true;
-				break;
+				if (SceneView->bCameraCut)
+				{
+					bCameraCutThisFrame = true;
+				}
+
+				if (SceneView->bIsSceneCapture)
+				{
+					bHasSceneCaptureView = true;
+				}
 			}
 		}
+		
 
-		AComposurePipelineBaseActor* Owner = AssociatedPipelineObj.Get();
-		Owner->EnqueueRendering(bCameraCutThisFrame);
+		// Note: Previously the engine wasn't calling BeginRenderViewFamily on scene captures.
+		// Now that it does, we explicitely disable this case to prevent recursive scene capture rendering.
+		if (!bHasSceneCaptureView)
+		{
+			AComposurePipelineBaseActor* Owner = AssociatedPipelineObj.Get();
+			Owner->EnqueueRendering(bCameraCutThisFrame);
+		}
 	}
 }
 
