Compositing: Fixed the issue with broken references to Target Cameras referenced in sublevels.

diff --git a/Engine/Plugins/Compositing/Composure/Source/Composure/Classes/CompositingElement.h b/Engine/Plugins/Compositing/Composure/Source/Composure/Classes/CompositingElement.h
--- a/Engine/Plugins/Compositing/Composure/Source/Composure/Classes/CompositingElement.h
+++ b/Engine/Plugins/Compositing/Composure/Source/Composure/Classes/CompositingElement.h
@@ -106,8 +106,11 @@ public:
 	UPROPERTY(EditAnywhere, Category = "Composure|Input")
 	ESceneCameraLinkType CameraSource = ESceneCameraLinkType::Unused;
 
-	UPROPERTY(EditInstanceOnly, Category = "Composure|Input")
-	TLazyObjectPtr<ACameraActor> TargetCameraActor;
+	UPROPERTY()	
+	TLazyObjectPtr<ACameraActor> TargetCameraActor_DEPRECATED;
+
+	UPROPERTY(EditInstanceOnly, Category = "Composure|Input", DisplayName = "Target Camera Actor")
+	TWeakObjectPtr<ACameraActor> TargetCameraActorPtr;
 
 	/*********************************/
 	// Outputs
diff --git a/Engine/Plugins/Compositing/Composure/Source/Composure/Private/CompositingCaptureBase.cpp b/Engine/Plugins/Compositing/Composure/Source/Composure/Private/CompositingCaptureBase.cpp
--- a/Engine/Plugins/Compositing/Composure/Source/Composure/Private/CompositingCaptureBase.cpp
+++ b/Engine/Plugins/Compositing/Composure/Source/Composure/Private/CompositingCaptureBase.cpp
@@ -117,10 +117,10 @@ void ACompositingCaptureBase::PostEditChangeProperty(struct FPropertyChangedEven
 {
 	const FName PropertyName = (PropertyChangedEvent.Property != NULL) ? PropertyChangedEvent.Property->GetFName() : NAME_None;
 
-	if (PropertyName == GET_MEMBER_NAME_CHECKED(ACompositingCaptureBase, TargetCameraActor))
+	if (PropertyName == GET_MEMBER_NAME_CHECKED(ACompositingCaptureBase, TargetCameraActorPtr))
 	{
 		// If there is no target camera, remove the last distortion post-process MID from the scene capture
-		if (TargetCameraActor == nullptr)
+		if (TargetCameraActorPtr == nullptr)
 		{
 			if (SceneCaptureComponent2D)
 			{
diff --git a/Engine/Plugins/Compositing/Composure/Source/Composure/Private/CompositingElement.cpp b/Engine/Plugins/Compositing/Composure/Source/Composure/Private/CompositingElement.cpp
--- a/Engine/Plugins/Compositing/Composure/Source/Composure/Private/CompositingElement.cpp
+++ b/Engine/Plugins/Compositing/Composure/Source/Composure/Private/CompositingElement.cpp
@@ -686,13 +686,13 @@ ACameraActor* ACompositingElement::FindTargetCamera() const
 {
 	if (CameraSource == ESceneCameraLinkType::Override)
 	{
-		return TargetCameraActor.Get();
+		return TargetCameraActorPtr.Get();
 	}
 	else if (Parent)
 	{
 		return Parent->FindTargetCamera();
 	}
-	else if (!TargetCameraActor.IsValid())
+	else if (!TargetCameraActorPtr.IsValid())
 	{
 		for (TActorIterator<ACineCameraActor> CamActorIt(GetWorld()); CamActorIt; ++CamActorIt)
 		{
@@ -712,7 +712,7 @@ ACameraActor* ACompositingElement::FindTargetCamera() const
 			}
 		}
 	}
-	return TargetCameraActor.Get();
+	return TargetCameraActorPtr.Get();
 }
 
 void ACompositingElement::SetTargetCamera(ACameraActor* NewCameraActor)
@@ -720,7 +720,7 @@ void ACompositingElement::SetTargetCamera(ACameraActor* NewCameraActor)
 	if (CameraSource == ESceneCameraLinkType::Override)
 	{
 		//Override mode has the highest priority in using the camera. Should check before going to the parent. 
-		TargetCameraActor = NewCameraActor;
+		TargetCameraActorPtr = NewCameraActor;
 	}
 	else if (Parent)
 	{
@@ -729,7 +729,7 @@ void ACompositingElement::SetTargetCamera(ACameraActor* NewCameraActor)
 	else 
 	{
 		//Base case for the recursive calls. Current element does not have parent and is not in override mode.
-		TargetCameraActor = NewCameraActor;
+		TargetCameraActorPtr = NewCameraActor;
 	}
 }
 
@@ -899,6 +899,16 @@ void ACompositingElement::PostLoad()
 	const int32 ComposureVer = GetLinkerCustomVersion(FComposureCustomVersion::GUID);
 	PostLoadCompatUpgrade(ComposureVer);
 
+#if WITH_EDITORONLY_DATA
+	// Re-assign the value of deprecated pointer to the new one.
+	// This will continue until user re-saves the level.
+	if (TargetCameraActor_DEPRECATED.IsValid())
+	{
+		TargetCameraActorPtr = TargetCameraActor_DEPRECATED.Get();
+		TargetCameraActor_DEPRECATED = nullptr;
+	}
+#endif
+
 	// Note that RefreshAllInternalPassLists is not called here because the refresh can call Reset, which is not allowed in PostLoad.
 }
 
diff --git a/Engine/Plugins/Compositing/Composure/Source/ComposureLayersEditor/Private/ComposureDetailCustomizations.cpp b/Engine/Plugins/Compositing/Composure/Source/ComposureLayersEditor/Private/ComposureDetailCustomizations.cpp
--- a/Engine/Plugins/Compositing/Composure/Source/ComposureLayersEditor/Private/ComposureDetailCustomizations.cpp
+++ b/Engine/Plugins/Compositing/Composure/Source/ComposureLayersEditor/Private/ComposureDetailCustomizations.cpp
@@ -363,7 +363,7 @@ void FCompElementDetailsCustomization::CustomizeDetails(IDetailLayoutBuilder& De
 
 		if (TRINARY_IS_FALSE(bShowCameraProp))
 		{
-			DetailBuilder.HideProperty(GET_MEMBER_NAME_CHECKED(ACompositingElement, TargetCameraActor));
+			DetailBuilder.HideProperty(GET_MEMBER_NAME_CHECKED(ACompositingElement, TargetCameraActorPtr));
 		}
 
 		TSharedPtr<IPropertyHandle> PreviewTransformSource = DetailBuilder.GetProperty(GET_MEMBER_NAME_CHECKED(ACompositingElement, PreviewTransformSource));
