Fixed clang Win64 editor compilation failures caused by ensure macros change

diff --git a/Engine/Plugins/Compositing/Composure/Source/Composure/Public/EditorSupport/CompFreezeFrameController.h b/Engine/Plugins/Compositing/Composure/Source/Composure/Public/EditorSupport/CompFreezeFrameController.h
--- a/Engine/Plugins/Compositing/Composure/Source/Composure/Public/EditorSupport/CompFreezeFrameController.h
+++ b/Engine/Plugins/Compositing/Composure/Source/Composure/Public/EditorSupport/CompFreezeFrameController.h
@@ -9,7 +9,7 @@
 enum class ETargetUsageFlags : uint8;
 typedef FGuid FFreezeFrameControlHandle;
 
-struct COMPOSURE_API FCompFreezeFrameController
+struct FCompFreezeFrameController
 {
 public:
 	FCompFreezeFrameController(int32& FreezeFlagsRef)
@@ -42,15 +42,15 @@ public:
 
 	FORCEINLINE ETargetUsageFlags GetFreezeFlags() const { return (ETargetUsageFlags)FreezeFlags; }
 
-	bool SetFreezeFlags(ETargetUsageFlags InFreezeFlags, bool bClearOthers = false, const FFreezeFrameControlHandle& LockKey = FFreezeFrameControlHandle());
-	bool ClearFreezeFlags(ETargetUsageFlags InFreezeFlags, const FFreezeFrameControlHandle& LockKey = FFreezeFrameControlHandle());
-	bool ClearFreezeFlags(const FFreezeFrameControlHandle& LockKey = FFreezeFrameControlHandle());
+	COMPOSURE_API bool SetFreezeFlags(ETargetUsageFlags InFreezeFlags, bool bClearOthers = false, const FFreezeFrameControlHandle& LockKey = FFreezeFrameControlHandle());
+	COMPOSURE_API bool ClearFreezeFlags(ETargetUsageFlags InFreezeFlags, const FFreezeFrameControlHandle& LockKey = FFreezeFrameControlHandle());
+	COMPOSURE_API bool ClearFreezeFlags(const FFreezeFrameControlHandle& LockKey = FFreezeFrameControlHandle());
 
 	FORCEINLINE operator ETargetUsageFlags() const { return GetFreezeFlags(); }
 	FORCEINLINE operator int32()                   { return (int32)GetFreezeFlags(); }
 
-	bool HasAnyFlags(ETargetUsageFlags InFreezeFlags);
-	bool HasAllFlags(ETargetUsageFlags InFreezeFlags);
+	COMPOSURE_API bool HasAnyFlags(ETargetUsageFlags InFreezeFlags);
+	COMPOSURE_API bool HasAllFlags(ETargetUsageFlags InFreezeFlags);
 
 private: 
 	FFreezeFrameControlHandle LockKey;
@@ -59,6 +59,6 @@ private:
 public:
 
 	/** DO NOT USE - For UObject construction only */
-	FCompFreezeFrameController(EForceInit Default = ForceInit);
+	COMPOSURE_API FCompFreezeFrameController(EForceInit Default = ForceInit);
 };
 
diff --git a/Engine/Source/Runtime/CoreUObject/Private/UObject/PropertySoftObjectPtr.cpp b/Engine/Source/Runtime/CoreUObject/Private/UObject/PropertySoftObjectPtr.cpp
--- a/Engine/Source/Runtime/CoreUObject/Private/UObject/PropertySoftObjectPtr.cpp
+++ b/Engine/Source/Runtime/CoreUObject/Private/UObject/PropertySoftObjectPtr.cpp
@@ -30,6 +30,19 @@ FString FSoftObjectProperty::GetCPPTypeCustom(FString* ExtendedTypeText, uint32
 	ensure(!InnerNativeTypeName.IsEmpty());
 	return FString::Printf(TEXT("TSoftObjectPtr<%s>"), *InnerNativeTypeName);
 }
+
+FString FSoftObjectProperty::GetCPPType(FString* ExtendedTypeText, uint32 CPPExportFlags) const
+{
+	if (ensureMsgf(PropertyClass, TEXT("Soft object property missing PropertyClass: %s"), *GetFullNameSafe(this)))
+	{
+		return Super::GetCPPType(ExtendedTypeText, CPPExportFlags);
+	}
+	else
+	{
+		return TEXT("TSoftObjectPtr<UObject>");
+	}
+}
+
 FString FSoftObjectProperty::GetCPPMacroType( FString& ExtendedTypeText ) const
 {
 	ExtendedTypeText = FString::Printf(TEXT("TSoftObjectPtr<%s%s>"), PropertyClass->GetPrefixCPP(), *PropertyClass->GetName());
diff --git a/Engine/Source/Runtime/CoreUObject/Public/UObject/UnrealType.h b/Engine/Source/Runtime/CoreUObject/Public/UObject/UnrealType.h
--- a/Engine/Source/Runtime/CoreUObject/Public/UObject/UnrealType.h
+++ b/Engine/Source/Runtime/CoreUObject/Public/UObject/UnrealType.h
@@ -3149,19 +3149,8 @@ public:
 	virtual void SetObjectPropertyValue(void* PropertyValueAddress, UObject* Value) const override;
 	virtual void SetObjectPropertyValue_InContainer(void* ContainerAddress, UObject* Value, int32 ArrayIndex = 0) const override;
 	virtual bool AllowCrossLevel() const override;
-	virtual FString GetCPPTypeCustom(FString* ExtendedTypeText, uint32 CPPExportFlags, const FString& InnerNativeTypeName)  const override;
-
-	virtual FString GetCPPType(FString* ExtendedTypeText, uint32 CPPExportFlags) const override
-	{
-		if (ensureMsgf(PropertyClass, TEXT("Soft object property missing PropertyClass: %s"), *GetFullNameSafe(this)))
-		{
-			return Super::GetCPPType(ExtendedTypeText, CPPExportFlags);
-		}
-		else
-		{
-			return TEXT("TSoftObjectPtr<UObject>");
-		}
-	}
+	virtual FString GetCPPTypeCustom(FString* ExtendedTypeText, uint32 CPPExportFlags, const FString& InnerNativeTypeName) const override;
+	virtual FString GetCPPType(FString* ExtendedTypeText, uint32 CPPExportFlags) const override;
 
 private:
 	virtual uint32 GetValueTypeHashInternal(const void* Src) const override;
