Composure: Further disable show flags for post-processing-only renders.

The primary issue this fixes is that lumen calculations were being performed and thrown away on the main empty render. Additional flags were turned off for good measure.

diff --git a/Engine/Plugins/Compositing/Composure/Source/Composure/Private/ComposureUtils.cpp b/Engine/Plugins/Compositing/Composure/Source/Composure/Private/ComposureUtils.cpp
--- a/Engine/Plugins/Compositing/Composure/Source/Composure/Private/ComposureUtils.cpp
+++ b/Engine/Plugins/Compositing/Composure/Source/Composure/Private/ComposureUtils.cpp
@@ -14,16 +14,21 @@ void FComposureUtils::SetEngineShowFlagsForPostprocessingOnly(FEngineShowFlags&
 	EngineShowFlags.ScreenSpaceAO = false;
 	EngineShowFlags.LightShafts = false;
 	EngineShowFlags.Lighting = false;
+	EngineShowFlags.SkyLighting = false;
 	EngineShowFlags.DeferredLighting = false;
 	EngineShowFlags.Decals = false;
 	EngineShowFlags.Translucency = false;
 	EngineShowFlags.AntiAliasing = false;
+	EngineShowFlags.TemporalAA = false;
 	EngineShowFlags.MotionBlur = false;
 	EngineShowFlags.Bloom = false;
 	EngineShowFlags.EyeAdaptation = false;
 	EngineShowFlags.VolumetricFog = false;
 	EngineShowFlags.Atmosphere = false;
 	EngineShowFlags.Fog = false;
+	EngineShowFlags.LocalExposure = false;
+	EngineShowFlags.GlobalIllumination = false;
+	EngineShowFlags.LumenReflections = false;
 
 #if !UE_BUILD_OPTIMIZED_SHOWFLAGS
 	// Development-only flags
