Composure: Fixed removal function which incorrectly removed transform elements instead of input elements.

diff --git a/Engine/Plugins/Compositing/Composure/Source/Composure/Private/CompositingElement.cpp b/Engine/Plugins/Compositing/Composure/Source/Composure/Private/CompositingElement.cpp
--- a/Engine/Plugins/Compositing/Composure/Source/Composure/Private/CompositingElement.cpp
+++ b/Engine/Plugins/Compositing/Composure/Source/Composure/Private/CompositingElement.cpp
@@ -547,7 +547,7 @@ int32 ACompositingElement::RemovePassesOfType(TSubclassOf<UCompositingElementPas
 
 	if (PassType->IsChildOf<UCompositingElementInput>())
 	{
-		RemoveCount = CompositingElement_Impl::RemovePassesOfType(TransformPasses, UserConstructedTransforms, PassType);
+		RemoveCount = CompositingElement_Impl::RemovePassesOfType(Inputs, UserConstructedInputs, PassType);
 		if (RemoveCount > 0)
 		{
 			RefreshInternalInputsList();
