Reinstate change type propagation when sending property change events up the chain, allowing Interactive editor changes to work on properties in nested arrays
- Previously, the change type was changed to Unspecified to prevent ArrayAdd events on nested arrays from being interpreted as adds
- Resolved the offending case by swapping `ACompositingElement::PostEditChangeProperty` to `ACompositingElement::PostEditChangeChainProperty` and checking that there are no nested containers further down the property chain

#tests [REDACTED_HANDLE] verified that the issue in UE-213813 is still resolved

diff --git a/Engine/Plugins/Compositing/Composure/Source/Composure/Classes/CompositingElement.h b/Engine/Plugins/Compositing/Composure/Source/Composure/Classes/CompositingElement.h
--- a/Engine/Plugins/Compositing/Composure/Source/Composure/Classes/CompositingElement.h
+++ b/Engine/Plugins/Compositing/Composure/Source/Composure/Classes/CompositingElement.h
@@ -383,7 +383,7 @@ public:
 	virtual void PostLoad() override;
 	virtual void BeginDestroy() override;
 #if WITH_EDITOR
-	virtual void PostEditChangeProperty(FPropertyChangedEvent& PropertyChangedEvent) override;
+	virtual void PostEditChangeChainProperty(struct FPropertyChangedChainEvent& PropertyChangedChainEvent) override;
 	virtual void PostEditUndo() override;
 	virtual void PostDuplicate(bool bDuplicateForPIE) override;
 #endif
diff --git a/Engine/Plugins/Compositing/Composure/Source/Composure/Private/EditorSupport/CompositingElementEditorSupport.cpp b/Engine/Plugins/Compositing/Composure/Source/Composure/Private/EditorSupport/CompositingElementEditorSupport.cpp
--- a/Engine/Plugins/Compositing/Composure/Source/Composure/Private/EditorSupport/CompositingElementEditorSupport.cpp
+++ b/Engine/Plugins/Compositing/Composure/Source/Composure/Private/EditorSupport/CompositingElementEditorSupport.cpp
@@ -113,9 +113,10 @@ FCompFreezeFrameController* ACompositingElement::GetFreezeFrameController()
 	return &FreezeFrameController;
 }
 
-void ACompositingElement::PostEditChangeProperty(FPropertyChangedEvent& PropertyChangedEvent)
+void ACompositingElement::PostEditChangeChainProperty(struct FPropertyChangedChainEvent& PropertyChangedChainEvent)
 {
-	const FName PropertyName = PropertyChangedEvent.GetPropertyName();
+	const FName PropertyName = PropertyChangedChainEvent.PropertyChain.GetActiveNode()->GetValue()->GetFName();
+	
 	if (PropertyName == TEXT("ActorLabel"))//GET_MEMBER_NAME_CHECKED(ACompositingElement, ActorLabel))
 	{
 		CompShotIdName = *GetActorLabel();
@@ -136,110 +137,107 @@ void ACompositingElement::PostEditChangeProperty(FPropertyChangedEvent& Property
 			OnDisabled();
 		}
 	}
-	else if (PropertyName == GET_MEMBER_NAME_CHECKED(ACompositingElement, Inputs))
+	else if (PropertyName == GET_MEMBER_NAME_CHECKED(ACompositingElement, bAutoRunChildElementsAndSelf))
 	{
-		if (!HasAnyFlags(RF_ClassDefaultObject))
-		{
-			if (PropertyChangedEvent.ChangeType == EPropertyChangeType::ArrayAdd && DefaultInputType)
-			{
-				const int32 AddedIndex = PropertyChangedEvent.GetArrayIndex(PropertyChangedEvent.GetPropertyName().ToString());
-				if (Inputs.IsValidIndex(AddedIndex))
-				{
-					Inputs[AddedIndex] = FCompositingElementPassUtils::NewInstancedSubObj<UCompositingElementInput>(/*Outer =*/this, DefaultInputType);
-					Inputs[AddedIndex]->PassName = MakeUniqueObjectName(this, UCompositingElementInput::StaticClass(), TEXT("InputPass"));
-				}
-			}
-			else if (PropertyChangedEvent.ChangeType == EPropertyChangeType::ValueSet)
-			{
-				const int32 AlteredIndex = PropertyChangedEvent.GetArrayIndex(PropertyChangedEvent.GetPropertyName().ToString());
-				if (Inputs.IsValidIndex(AlteredIndex) && Inputs[AlteredIndex] && Inputs[AlteredIndex]->PassName.IsNone())
-				{
-					if (UCompositingElementInput* ReplacedInput = CompositingElementEditorSupport_Impl::FindReplacedPass(Inputs, GetInternalInputsList(), AlteredIndex))
-					{
-						Inputs[AlteredIndex]->PassName = ReplacedInput->PassName;
-					}
-					else if (!GetInternalInputsList().Contains(Inputs[AlteredIndex]))
-					{
-						Inputs[AlteredIndex]->PassName = MakeUniqueObjectName(this, UCompositingElementInput::StaticClass(), TEXT("InputPass"));
-					}
-				}			
-			}
-			RefreshInternalInputsList();
-		}
+		SetAutoRunChildrenAndSelf(bAutoRunChildElementsAndSelf);
 	}
-	else if (PropertyName == GET_MEMBER_NAME_CHECKED(ACompositingElement, TransformPasses))
+	else if (!HasAnyFlags(RF_ClassDefaultObject))
 	{
-		if (!HasAnyFlags(RF_ClassDefaultObject))
+		const bool bEditedInputsArray = PropertyName == GET_MEMBER_NAME_CHECKED(ACompositingElement, Inputs);
+		const bool bEditedTransformPassesArray = PropertyName == GET_MEMBER_NAME_CHECKED(ACompositingElement, TransformPasses);
+		const bool bEditedOutputsArray = PropertyName == GET_MEMBER_NAME_CHECKED(ACompositingElement, Outputs);
+		if ((bEditedInputsArray || bEditedTransformPassesArray || bEditedOutputsArray))
 		{
-			if (PropertyChangedEvent.ChangeType == EPropertyChangeType::ArrayAdd && DefaultTransformType)
+			if (PropertyChangedChainEvent.ChangeType == EPropertyChangeType::ValueSet ||
+				(PropertyChangedChainEvent.ChangeType == EPropertyChangeType::ArrayAdd &&
+				((bEditedInputsArray && DefaultInputType) || (bEditedTransformPassesArray && DefaultTransformType) || (bEditedOutputsArray && DefaultOutputType))))
 			{
-				const int32 AddedIndex = PropertyChangedEvent.GetArrayIndex(PropertyChangedEvent.GetPropertyName().ToString());
-				if (TransformPasses.IsValidIndex(AddedIndex))
+				// Verify that there isn't another container property node further down in the chain (so we don't respond to array change events from a different array!)
+				bool bWasNestedContainerChanged = false;
+				FEditPropertyChain::TDoubleLinkedListNode* Node = PropertyChangedChainEvent.PropertyChain.GetActiveNode()->GetNextNode();
+				while (!bWasNestedContainerChanged && Node && Node->GetValue())
 				{
-					TransformPasses[AddedIndex] = FCompositingElementPassUtils::NewInstancedSubObj<UCompositingElementTransform>(/*Outer =*/this, DefaultTransformType);
-					TransformPasses[AddedIndex]->PassName = MakeUniqueObjectName(this, UCompositingElementTransform::StaticClass(), TEXT("TransformPass"));
+					bWasNestedContainerChanged = Node->GetValue()->IsA<FArrayProperty>() || Node->GetValue()->IsA<FSetProperty>() || Node->GetValue()->IsA<FMapProperty>();
+					Node = Node->GetNextNode();
 				}
-			}
-			else if (PropertyChangedEvent.ChangeType == EPropertyChangeType::ValueSet)
-			{
-				const int32 AlteredIndex = PropertyChangedEvent.GetArrayIndex(PropertyChangedEvent.GetPropertyName().ToString());
-				if (TransformPasses.IsValidIndex(AlteredIndex) && TransformPasses[AlteredIndex] && TransformPasses[AlteredIndex]->PassName.IsNone())
+
+				if (!bWasNestedContainerChanged)
 				{
-					if (UCompositingElementTransform* ReplacedInput = CompositingElementEditorSupport_Impl::FindReplacedPass(TransformPasses, GetInternalTransformsList(), AlteredIndex))
+					const int32 ArrayIndex = PropertyChangedChainEvent.GetArrayIndex(PropertyName.ToString());
+					if (bEditedInputsArray && Inputs.IsValidIndex(ArrayIndex))
 					{
-						TransformPasses[AlteredIndex]->PassName = ReplacedInput->PassName;
+						if (PropertyChangedChainEvent.ChangeType == EPropertyChangeType::ArrayAdd)
+						{
+							Inputs[ArrayIndex] = FCompositingElementPassUtils::NewInstancedSubObj<UCompositingElementInput>(/*Outer =*/this, DefaultInputType);
+							Inputs[ArrayIndex]->PassName = MakeUniqueObjectName(this, UCompositingElementInput::StaticClass(), TEXT("InputPass"));
+						}
+						else if (Inputs[ArrayIndex] && Inputs[ArrayIndex]->PassName.IsNone())
+						{
+							if (UCompositingElementInput* ReplacedInput = CompositingElementEditorSupport_Impl::FindReplacedPass(Inputs, GetInternalInputsList(), ArrayIndex))
+							{
+								Inputs[ArrayIndex]->PassName = ReplacedInput->PassName;
+							}
+							else if (!GetInternalInputsList().Contains(Inputs[ArrayIndex]))
+							{
+								Inputs[ArrayIndex]->PassName = MakeUniqueObjectName(this, UCompositingElementInput::StaticClass(), TEXT("InputPass"));
+							}
+						}
+						
+						RefreshInternalInputsList();
 					}
-					else if (!GetInternalTransformsList().Contains(TransformPasses[AlteredIndex]))
+					else if (bEditedTransformPassesArray && TransformPasses.IsValidIndex(ArrayIndex))
 					{
-						TransformPasses[AlteredIndex]->PassName = MakeUniqueObjectName(this, UCompositingElementTransform::StaticClass(), TEXT("TransformPass"));
+						if (PropertyChangedChainEvent.ChangeType == EPropertyChangeType::ArrayAdd)
+						{
+							TransformPasses[ArrayIndex] = FCompositingElementPassUtils::NewInstancedSubObj<UCompositingElementTransform>(/*Outer =*/this, DefaultTransformType);
+							TransformPasses[ArrayIndex]->PassName = MakeUniqueObjectName(this, UCompositingElementTransform::StaticClass(), TEXT("TransformPass"));
+						}
+						else if (TransformPasses[ArrayIndex] && TransformPasses[ArrayIndex]->PassName.IsNone())
+						{
+							if (UCompositingElementTransform* ReplacedInput = CompositingElementEditorSupport_Impl::FindReplacedPass(TransformPasses, GetInternalTransformsList(), ArrayIndex))
+							{
+								TransformPasses[ArrayIndex]->PassName = ReplacedInput->PassName;
+							}
+							else if (!GetInternalTransformsList().Contains(TransformPasses[ArrayIndex]))
+							{
+								TransformPasses[ArrayIndex]->PassName = MakeUniqueObjectName(this, UCompositingElementTransform::StaticClass(), TEXT("TransformPass"));
+							}
+						}
+						
+						RefreshInternalTransformsList();
 					}
-				}
-			}
-			RefreshInternalTransformsList();
-		}
-	}
-	else if (PropertyName == GET_MEMBER_NAME_CHECKED(ACompositingElement, Outputs))
-	{
-		if (!HasAnyFlags(RF_ClassDefaultObject))
-		{
-			if (PropertyChangedEvent.ChangeType == EPropertyChangeType::ArrayAdd && DefaultOutputType)
-			{
-				const int32 AddedIndex = PropertyChangedEvent.GetArrayIndex(PropertyChangedEvent.GetPropertyName().ToString());
-				if (Outputs.IsValidIndex(AddedIndex))
-				{
-					Outputs[AddedIndex] = FCompositingElementPassUtils::NewInstancedSubObj<UCompositingElementOutput>(/*Outer =*/this, DefaultOutputType);
-					Outputs[AddedIndex]->PassName = MakeUniqueObjectName(this, UCompositingElementOutput::StaticClass(), TEXT("OutputPass"));
-				}
-			}
-			else if (PropertyChangedEvent.ChangeType == EPropertyChangeType::ValueSet)
-			{
-				const int32 AlteredIndex = PropertyChangedEvent.GetArrayIndex(PropertyChangedEvent.GetPropertyName().ToString());
-				if (Outputs.IsValidIndex(AlteredIndex) && Outputs[AlteredIndex] && Outputs[AlteredIndex]->PassName.IsNone())
-				{
-					if (UCompositingElementOutput* ReplacedInput = CompositingElementEditorSupport_Impl::FindReplacedPass(Outputs, GetInternalOutputsList(), AlteredIndex))
+					else if (bEditedOutputsArray && Outputs.IsValidIndex(ArrayIndex))
 					{
-						Outputs[AlteredIndex]->PassName = ReplacedInput->PassName;
-					}
-					else if (!GetInternalOutputsList().Contains(Outputs[AlteredIndex]))
-					{
-						Outputs[AlteredIndex]->PassName = MakeUniqueObjectName(this, UCompositingElementOutput::StaticClass(), TEXT("OutputPass"));
+						if (PropertyChangedChainEvent.ChangeType == EPropertyChangeType::ArrayAdd)
+						{
+							Outputs[ArrayIndex] = FCompositingElementPassUtils::NewInstancedSubObj<UCompositingElementOutput>(/*Outer =*/this, DefaultOutputType);
+							Outputs[ArrayIndex]->PassName = MakeUniqueObjectName(this, UCompositingElementOutput::StaticClass(), TEXT("OutputPass"));
+						}
+						else if (Outputs[ArrayIndex] && Outputs[ArrayIndex]->PassName.IsNone())
+						{
+							if (UCompositingElementOutput* ReplacedInput = CompositingElementEditorSupport_Impl::FindReplacedPass(Outputs, GetInternalOutputsList(), ArrayIndex))
+							{
+								Outputs[ArrayIndex]->PassName = ReplacedInput->PassName;
+							}
+							else if (!GetInternalOutputsList().Contains(Outputs[ArrayIndex]))
+							{
+								Outputs[ArrayIndex]->PassName = MakeUniqueObjectName(this, UCompositingElementOutput::StaticClass(), TEXT("OutputPass"));
+							}
+						}
+						
+						RefreshInternalOutputsList();
 					}
 				}
 			}
-			RefreshInternalOutputsList();
 		}
 	}
-	else if (PropertyName == GET_MEMBER_NAME_CHECKED(ACompositingElement, bAutoRunChildElementsAndSelf))
-	{
-		SetAutoRunChildrenAndSelf(bAutoRunChildElementsAndSelf);
-	}
 
 	if (ICompositingEditor* CompositingEditor = ICompositingEditor::Get())
 	{
 		CompositingEditor->RequestRedraw();
 	}
-
-	Super::PostEditChangeProperty(PropertyChangedEvent);
+	
+	Super::PostEditChangeChainProperty(PropertyChangedChainEvent);
 }
 
 void ACompositingElement::PostEditUndo()
diff --git a/Engine/Source/Editor/PropertyEditor/Private/PropertyNode.cpp b/Engine/Source/Editor/PropertyEditor/Private/PropertyNode.cpp
--- a/Engine/Source/Editor/PropertyEditor/Private/PropertyNode.cpp
+++ b/Engine/Source/Editor/PropertyEditor/Private/PropertyNode.cpp
@@ -3217,8 +3217,6 @@ void FPropertyNode::NotifyPostChange( FPropertyChangedEvent& InPropertyChangedEv
 					FPropertyChangedEvent ChangedEvent = InPropertyChangedEvent;
 					if (CurProperty != InPropertyChangedEvent.Property)
 					{
-						// Parent object node property. Reset property and leave the event type as unspecified since we dont pass in the exact leaf property.
-						ChangedEvent.ChangeType = EPropertyChangeType::Unspecified;
 						ChangedEvent.Property = CurProperty;
 						ChangedEvent.MemberProperty = CurProperty;
 					}
