Fix missing virtual destructor after the base interface class had its own removed

#rnx

diff --git a/Engine/Plugins/Compositing/Composure/Source/Composure/Private/MovieScene/MovieSceneComposureExportSectionTemplate.cpp b/Engine/Plugins/Compositing/Composure/Source/Composure/Private/MovieScene/MovieSceneComposureExportSectionTemplate.cpp
--- a/Engine/Plugins/Compositing/Composure/Source/Composure/Private/MovieScene/MovieSceneComposureExportSectionTemplate.cpp
+++ b/Engine/Plugins/Compositing/Composure/Source/Composure/Private/MovieScene/MovieSceneComposureExportSectionTemplate.cpp
@@ -55,6 +55,8 @@ struct FExportIntermediateBuffersViewExtension : ISceneViewExtension
 	 */
 	static TSharedRef<FExportIntermediateBuffersViewExtension, ESPMode::ThreadSafe> Create(ACompositingElement* InCompShotElement, USceneCaptureComponent2D* InSceneCapture, const TArray<FString>& InBuffersToCapture);
 
+	virtual ~FExportIntermediateBuffersViewExtension() = default;
+
 private:
 
 	virtual void SetupViewFamily( FSceneViewFamily& InViewFamily ) override;
diff --git a/Engine/Plugins/Experimental/CineCameraSceneCapture/Source/CineCameraSceneCapture/Private/CineCameraSceneCaptureComponent.cpp b/Engine/Plugins/Experimental/CineCameraSceneCapture/Source/CineCameraSceneCapture/Private/CineCameraSceneCaptureComponent.cpp
--- a/Engine/Plugins/Experimental/CineCameraSceneCapture/Source/CineCameraSceneCapture/Private/CineCameraSceneCaptureComponent.cpp
+++ b/Engine/Plugins/Experimental/CineCameraSceneCapture/Source/CineCameraSceneCapture/Private/CineCameraSceneCaptureComponent.cpp
@@ -32,6 +32,8 @@ public:
 		, bIsOcioEnabledRenderThread(false)
 	{ }
 
+	virtual ~FCineCameraCaptureSceneViewExtension() = default;
+
 	//~ Begin FSceneViewExtensionBase Interface
 	virtual void SetupViewFamily(FSceneViewFamily& InViewFamily) override {};
 	virtual void SetupView(FSceneViewFamily& InViewFamily, FSceneView& InView) override
