Composite: Preliminary renaming of display/friendly names to recycle Composure branding.

diff --git a/Engine/Plugins/Compositing/Composure/Composure.uplugin b/Engine/Plugins/Compositing/Composure/Composure.uplugin
--- a/Engine/Plugins/Compositing/Composure/Composure.uplugin
+++ b/Engine/Plugins/Compositing/Composure/Composure.uplugin
@@ -2,8 +2,8 @@
 	"FileVersion" : 3,
 	"Version" : 1,
 	"VersionName" : "1.0",
-	"FriendlyName" : "Composure",
-	"Description" : "Plugin design to make compositing easier in Unreal Engine.",
+	"FriendlyName" : "Legacy Composure",
+	"Description" : "Legacy system for real-time compositing. This plugin is no longer developed. Use Composure going forward.",
 	"Category" : "Compositing",
 	"CreatedBy" : "Epic Games, Inc.",
 	"CreatedByURL" : "https://epicgames.com",
diff --git a/Engine/Plugins/Compositing/Composure/Source/ComposureLayersEditor/Private/CompElementEditorModule.cpp b/Engine/Plugins/Compositing/Composure/Source/ComposureLayersEditor/Private/CompElementEditorModule.cpp
--- a/Engine/Plugins/Compositing/Composure/Source/ComposureLayersEditor/Private/CompElementEditorModule.cpp
+++ b/Engine/Plugins/Compositing/Composure/Source/ComposureLayersEditor/Private/CompElementEditorModule.cpp
@@ -281,8 +281,8 @@ void FCompElementEditorModule::RegisterEditorTab()
 
 		const FSlateIcon LayersIcon(FAppStyle::GetAppStyleSetName(), "LevelEditor.Tabs.ComposureCompositing");
 		LevelEditorTabManager->RegisterTabSpawner(CompElementEditor_Impl::ComposureLayersTabName, FOnSpawnTab::CreateStatic(&FCompElementEditorModule::SpawnComposureLayersTab))
-			.SetDisplayName(NSLOCTEXT("LevelEditorTabs", "LevelEditorComposureLayerBrowser", "Composure Compositing"))
-			.SetTooltipText(NSLOCTEXT("LevelEditorTabs", "LevelEditorComposureLayerBrowserTooltipText", "Open the Composure compositing tab."))
+			.SetDisplayName(NSLOCTEXT("LevelEditorTabs", "LevelEditorComposureLayerBrowser", "Legacy Composure"))
+			.SetTooltipText(NSLOCTEXT("LevelEditorTabs", "LevelEditorComposureLayerBrowserTooltipText", "Open the legacy Composure compositing tab."))
 			.SetGroup(WorkspaceMenu::GetMenuStructure().GetLevelEditorVirtualProductionCategory())
 			.SetIcon(LayersIcon);
 	});
diff --git a/Engine/Plugins/Experimental/Compositing/Composite/Composite.uplugin b/Engine/Plugins/Experimental/Compositing/Composite/Composite.uplugin
--- a/Engine/Plugins/Experimental/Compositing/Composite/Composite.uplugin
+++ b/Engine/Plugins/Experimental/Compositing/Composite/Composite.uplugin
@@ -2,8 +2,8 @@
 	"FileVersion": 1,
 	"Version": 1,
 	"VersionName": "1.0",
-	"FriendlyName": "Composite",
-	"Description": "Plugin designed to facilitate real-time compositing pipelines, extending CompositeCore.",
+	"FriendlyName": "Composure",
+	"Description": "Modern system for real-time compositing. This plugin succeeds legacy Composure and extends CompositeCore.",
 	"Category": "Compositing",
 	"CreatedBy": "Epic Games, Inc.",
 	"CreatedByURL": "https://epicgames.com",
